Changed category of Virtual Camera plugin to Virtual Production #RNX #Jira UE-59404

diff --git a/Engine/Plugins/Experimental/VirtualCamera/VirtualCamera.uplugin b/Engine/Plugins/Experimental/VirtualCamera/VirtualCamera.uplugin
--- a/Engine/Plugins/Experimental/VirtualCamera/VirtualCamera.uplugin
+++ b/Engine/Plugins/Experimental/VirtualCamera/VirtualCamera.uplugin
@@ -4,7 +4,7 @@
 	"VersionName": "1.0",
 	"FriendlyName": "VirtualCamera",
 	"Description": "",
-	"Category": "Other",
+	"Category": "Virtual Production",
 	"CreatedBy": "",
 	"CreatedByURL": "",
 	"DocsURL": "",
