Fix for Virtual Camera not getting correct aspect ratio when recording to sequencer #rb none #jira none

diff --git a/Engine/Plugins/Experimental/VirtualCamera/Source/VirtualCamera/Private/LevelSequencePlaybackController.cpp b/Engine/Plugins/Experimental/VirtualCamera/Source/VirtualCamera/Private/LevelSequencePlaybackController.cpp
--- a/Engine/Plugins/Experimental/VirtualCamera/Source/VirtualCamera/Private/LevelSequencePlaybackController.cpp
+++ b/Engine/Plugins/Experimental/VirtualCamera/Source/VirtualCamera/Private/LevelSequencePlaybackController.cpp
@@ -176,7 +176,7 @@ void ULevelSequencePlaybackController::ClearActiveLevelSequence()
 	}
 }
 
-void ULevelSequencePlaybackController::PilotTargetedCamera()
+void ULevelSequencePlaybackController::PilotTargetedCamera(FCameraFilmbackSettings* FilmbackSettingsOverride)
 {
 	if (!TargetCamera)
 	{
@@ -188,8 +188,8 @@ void ULevelSequencePlaybackController::PilotTargetedCamera()
 
 	TargetComponent->FocusSettings = CameraToFollow->FocusSettings;
 	TargetComponent->CurrentFocalLength = CameraToFollow->CurrentFocalLength;
-	TargetComponent->FilmbackSettings = CameraToFollow->FilmbackSettings;
 	TargetComponent->LensSettings = CameraToFollow->LensSettings;
+	TargetComponent->FilmbackSettings = FilmbackSettingsOverride ? *FilmbackSettingsOverride : CameraToFollow->FilmbackSettings;
 }
 
 void ULevelSequencePlaybackController::PlayFromBeginning()
diff --git a/Engine/Plugins/Experimental/VirtualCamera/Source/VirtualCamera/Private/VirtualCameraPlayerControllerBase.cpp b/Engine/Plugins/Experimental/VirtualCamera/Source/VirtualCamera/Private/VirtualCameraPlayerControllerBase.cpp
--- a/Engine/Plugins/Experimental/VirtualCamera/Source/VirtualCamera/Private/VirtualCameraPlayerControllerBase.cpp
+++ b/Engine/Plugins/Experimental/VirtualCamera/Source/VirtualCamera/Private/VirtualCameraPlayerControllerBase.cpp
@@ -85,7 +85,7 @@ void AVirtualCameraPlayerControllerBase::BeginPlay()
 void AVirtualCameraPlayerControllerBase::Tick(float DeltaSeconds)
 {
 	Super::Tick(DeltaSeconds);
-
+	
 	if (AVirtualCameraPawnBase* VCPawn = GetVirtualCameraPawn())
 	{
 		UpdatePawnWithTrackerData();
@@ -95,23 +95,24 @@ void AVirtualCameraPlayerControllerBase::Tick(float DeltaSeconds)
 			// Auto focus is just setting screen focus through the Auto Focus position.
 			SetFocusDistanceThroughPoint(AutoFocusScreenPosition);
 		}
+	}
 
-		if (UVirtualCameraCineCameraComponent* VPCamera = GetVirtualCameraCineCameraComponent())
+	UVirtualCameraCineCameraComponent* VCCamera = GetVirtualCameraCineCameraComponent();
+	if (VCCamera)
+	{
+		if (VCCamera->GetCurrentFocusMethod() == EVirtualCameraFocusMethod::Auto)
 		{
-			if (VPCamera->GetCurrentFocusMethod() == EVirtualCameraFocusMethod::Auto)
-			{
-				SetFocusDistanceThroughPoint(AutoFocusScreenPosition);
-			}
-
-			VPCamera->UpdateCameraView();
+			SetFocusDistanceThroughPoint(AutoFocusScreenPosition);
 		}
+
+		VCCamera->UpdateCameraView();
 	}
 
 	if (LevelSequencePlaybackController)
 	{
 		if (LevelSequencePlaybackController->bIsRecording)
 		{
-			LevelSequencePlaybackController->PilotTargetedCamera();
+			LevelSequencePlaybackController->PilotTargetedCamera(VCCamera ? &VCCamera->DesiredFilmbackSettings : nullptr);
 		}
 
 		if (LevelSequencePlaybackController->GetSequence())
diff --git a/Engine/Plugins/Experimental/VirtualCamera/Source/VirtualCamera/Public/LevelSequencePlaybackController.h b/Engine/Plugins/Experimental/VirtualCamera/Source/VirtualCamera/Public/LevelSequencePlaybackController.h
--- a/Engine/Plugins/Experimental/VirtualCamera/Source/VirtualCamera/Public/LevelSequencePlaybackController.h
+++ b/Engine/Plugins/Experimental/VirtualCamera/Source/VirtualCamera/Public/LevelSequencePlaybackController.h
@@ -92,8 +92,9 @@ public:
 
 	/**
 	 * Pilot the controlled camera during recording, copying over settings from the pawn.
+	 * @param FilmbackSettings - Optional override for filmback settings 
 	 */
-	void PilotTargetedCamera();
+	void PilotTargetedCamera(FCameraFilmbackSettings* FilmbackSettingsOverride = nullptr);
 
 	/**
 	 * Set the target camera component we need to be tracking during recording
