Merging fix from //UE4/Dev-Enterprise CL#7227034 : Fix focus stealing when starting/stopping recording with the VCAM

Merging fix from //UE4/Dev-Enterprise CL#7227105 : Bring game window back to front when starting/stopping recording with the VCAM - blueprint change

diff --git a/Engine/Plugins/Experimental/VirtualCamera/Source/VirtualCamera/Private/VirtualCameraPlayerControllerBase.cpp b/Engine/Plugins/Experimental/VirtualCamera/Source/VirtualCamera/Private/VirtualCameraPlayerControllerBase.cpp
--- a/Engine/Plugins/Experimental/VirtualCamera/Source/VirtualCamera/Private/VirtualCameraPlayerControllerBase.cpp
+++ b/Engine/Plugins/Experimental/VirtualCamera/Source/VirtualCamera/Private/VirtualCameraPlayerControllerBase.cpp
@@ -7,6 +7,7 @@
 #include "Engine/Engine.h"
 #include "Engine/World.h"
 #include "Features/IModularFeatures.h"
+#include "Framework/Application/SlateApplication.h"
 #include "IXRTrackingSystem.h"
 #include "RemoteSession/RemoteSession.h"
 #include "RemoteSession/Channels/RemoteSessionInputChannel.h"
@@ -1390,3 +1391,8 @@ bool AVirtualCameraPlayerControllerBase::ToggleAxisLock(const EVirtualCameraAxis
 
 	return false;
 }
+
+void AVirtualCameraPlayerControllerBase::ActivateGameViewport()
+{
+	FSlateApplication::Get().ActivateGameViewport();
+}
diff --git a/Engine/Plugins/Experimental/VirtualCamera/Source/VirtualCamera/Public/VirtualCameraPlayerControllerBase.h b/Engine/Plugins/Experimental/VirtualCamera/Source/VirtualCamera/Public/VirtualCameraPlayerControllerBase.h
--- a/Engine/Plugins/Experimental/VirtualCamera/Source/VirtualCamera/Public/VirtualCameraPlayerControllerBase.h
+++ b/Engine/Plugins/Experimental/VirtualCamera/Source/VirtualCamera/Public/VirtualCameraPlayerControllerBase.h
@@ -942,4 +942,10 @@ public:
 	 */
 	UFUNCTION(BlueprintCallable, Category = "Virtual Camera|Movement")
 	bool ToggleAxisLock(const EVirtualCameraAxis AxisToToggle);
+
+	/**
+	 * Helper to bring the game viewport to front.
+	 */
+	UFUNCTION(BlueprintCallable, Category = "Virtual Camera")
+	void ActivateGameViewport();
 };
diff --git a/Engine/Plugins/Experimental/VirtualCamera/Source/VirtualCamera/VirtualCamera.Build.cs b/Engine/Plugins/Experimental/VirtualCamera/Source/VirtualCamera/VirtualCamera.Build.cs
--- a/Engine/Plugins/Experimental/VirtualCamera/Source/VirtualCamera/VirtualCamera.Build.cs
+++ b/Engine/Plugins/Experimental/VirtualCamera/Source/VirtualCamera/VirtualCamera.Build.cs
@@ -26,6 +26,13 @@ public class VirtualCamera : ModuleRules
 			}
 		);
 
+		PrivateDependencyModuleNames.AddRange(
+			new string[]
+			{
+				"Slate",
+			}
+		);
+
 		if (Target.bBuildDeveloperTools)
 		{
 			PrivateDefinitions.Add("VIRTUALCAMERA_WITH_CONCERT=1");
