Fix VirtualCamera compilation warning.

diff --git a/Engine/Plugins/Experimental/VirtualCamera/Source/VirtualCamera/Public/VirtualCameraSaveGame.h b/Engine/Plugins/Experimental/VirtualCamera/Source/VirtualCamera/Public/VirtualCameraSaveGame.h
--- a/Engine/Plugins/Experimental/VirtualCamera/Source/VirtualCamera/Public/VirtualCameraSaveGame.h
+++ b/Engine/Plugins/Experimental/VirtualCamera/Source/VirtualCamera/Public/VirtualCameraSaveGame.h
@@ -175,7 +175,7 @@ struct FVirtualCameraSettings
 		FilmbackWidth = 24.889f;
 		FilmbackHeight = 14.f;
 		bAllowFocusVisualization = true;
-		DebugFocusPlaneColor = FColor();
+		DebugFocusPlaneColor = FColor(EForceInit::ForceInit);
 		MatteOpacity = .7f;
 		DesiredDistanceUnits = EUnit::Meters;
 
@@ -217,6 +217,7 @@ struct FVirtualCameraWaypoint
 
 	FVirtualCameraWaypoint()
 	{
+		DateCreated = FDateTime(0);
 		bIsHomeWaypoint = false;
 		bIsFavorited = false;
 		Name = "";
@@ -282,6 +283,7 @@ struct FVirtualCameraSettingsPreset
 		bIsAxisLockingSettingsSaved = false;
 		bIsMotionScaleSettingsSaved = false;
 		bIsFavorited = false;
+		DateCreated = FDateTime(0);
 	};
 };
 
