Ignore FLevelSequenceData.LastEdited for member initialization test. The value is always initialized to the current time.

diff --git a/Engine/Plugins/Experimental/VirtualCamera/Source/VirtualCamera/Private/LevelSequencePlaybackController.cpp b/Engine/Plugins/Experimental/VirtualCamera/Source/VirtualCamera/Private/LevelSequencePlaybackController.cpp
--- a/Engine/Plugins/Experimental/VirtualCamera/Source/VirtualCamera/Private/LevelSequencePlaybackController.cpp
+++ b/Engine/Plugins/Experimental/VirtualCamera/Source/VirtualCamera/Private/LevelSequencePlaybackController.cpp
@@ -120,7 +120,7 @@ void ULevelSequencePlaybackController::GetLevelSequences(TArray<FLevelSequenceDa
 	}
 
 	// Sort the level sequence names alphabetically by Display Name
-	OutLevelSequenceNames.Sort([](const FLevelSequenceData LeftItem, const FLevelSequenceData RightItem)
+	OutLevelSequenceNames.Sort([](const FLevelSequenceData& LeftItem, const FLevelSequenceData& RightItem)
 	{
 		return LeftItem.DisplayName < RightItem.DisplayName;
 	});
diff --git a/Engine/Plugins/Experimental/VirtualCamera/Source/VirtualCamera/Public/LevelSequencePlaybackController.h b/Engine/Plugins/Experimental/VirtualCamera/Source/VirtualCamera/Public/LevelSequencePlaybackController.h
--- a/Engine/Plugins/Experimental/VirtualCamera/Source/VirtualCamera/Public/LevelSequencePlaybackController.h
+++ b/Engine/Plugins/Experimental/VirtualCamera/Source/VirtualCamera/Public/LevelSequencePlaybackController.h
@@ -27,7 +27,7 @@ public:
 	UPROPERTY(EditAnywhere, BlueprintReadWrite, Category = "Level Sequence Info")
 	FString DisplayName;
 
-	UPROPERTY(EditAnywhere, BlueprintReadWrite, Category = "Level Sequence Info")
+	UPROPERTY(EditAnywhere, BlueprintReadWrite, Category = "Level Sequence Info", meta = (IgnoreForMemberInitializationTest))
 	FDateTime LastEdited;
 
 	FLevelSequenceData(const FString InAssetPath = "", const FString InDisplayName = "", const FDateTime InLastEdited = FDateTime::Now())
