Add virtual production roles

diff --git a/Engine/Plugins/Experimental/VirtualCamera/Source/VirtualCamera/Private/ConcertVirtualCamera.cpp b/Engine/Plugins/Experimental/VirtualCamera/Source/VirtualCamera/Private/ConcertVirtualCamera.cpp
new file mode 100644
--- /dev/null
+++ b/Engine/Plugins/Experimental/VirtualCamera/Source/VirtualCamera/Private/ConcertVirtualCamera.cpp
@@ -0,0 +1,161 @@
+// Copyright 1998-2019 Epic Games, Inc. All Rights Reserved.
+
+#include "ConcertVirtualCamera.h"
+
+#if VIRTUALCAMERA_WITH_CONCERT
+#include "IConcertClient.h"
+#include "IConcertModule.h"
+#include "IConcertSession.h"
+#endif
+
+
+/**
+ *
+ */
+FConcertVirtualCameraCameraFocusData::FConcertVirtualCameraCameraFocusData()
+	: ManualFocusDistance(100000.f)
+	, FocusSmoothingInterpSpeed(8.f)
+	, bSmoothFocusChanges(false)
+{}
+
+
+FConcertVirtualCameraCameraFocusData::FConcertVirtualCameraCameraFocusData(const UCineCameraComponent* CineCamera)
+	: ManualFocusDistance(CineCamera->CurrentFocusDistance)
+	, FocusSmoothingInterpSpeed(CineCamera->FocusSettings.FocusSmoothingInterpSpeed)
+	, bSmoothFocusChanges(CineCamera->FocusSettings.bSmoothFocusChanges)
+{}
+
+
+FCameraFocusSettings FConcertVirtualCameraCameraFocusData::ToCameraFocusSettings() const
+{
+	FCameraFocusSettings Result;
+	Result.FocusMethod = ECameraFocusMethod::Manual;
+	Result.ManualFocusDistance = ManualFocusDistance;
+	Result.bSmoothFocusChanges = bSmoothFocusChanges;
+	Result.FocusSmoothingInterpSpeed = FocusSmoothingInterpSpeed;
+	Result.FocusOffset = 0.f;
+	return Result;
+}
+
+
+/**
+ *
+ */
+FConcertVirtualCameraCameraEvent::FConcertVirtualCameraCameraEvent()
+	: InputSource(ETrackerInputSource::ARKit)
+	, CameraActorLocation(FVector::ZeroVector)
+	, CameraActorRotation(FRotator::ZeroRotator)
+	, CameraComponentLocation(FVector::ZeroVector)
+	, CameraComponentRotation(FRotator::ZeroRotator)
+	, CurrentAperture(0.f)
+	, CurrentFocalLength(0.f)
+{}
+
+
+#if VIRTUALCAMERA_WITH_CONCERT
+
+/**
+ *
+ */
+FConcertVirtualCameraManager::FConcertVirtualCameraManager()
+	: bIsLatestCameraEventDataValid(false)
+{
+	IConcertClientPtr ConcertClient = IConcertModule::Get().GetClientInstance();
+	if (ConcertClient.IsValid())
+	{
+		OnSessionStartupHandle = ConcertClient->OnSessionStartup().AddRaw(this, &FConcertVirtualCameraManager::RegisterConcertSyncHandlers);
+		OnSessionShutdownHandle = ConcertClient->OnSessionShutdown().AddRaw(this, &FConcertVirtualCameraManager::UnregisterConcertSyncHandlers);
+
+		TSharedPtr<IConcertClientSession> ConcertClientSession = ConcertClient->GetCurrentSession();
+		if (ConcertClientSession.IsValid())
+		{
+			RegisterConcertSyncHandlers(ConcertClientSession.ToSharedRef());
+		}
+	}
+}
+
+
+FConcertVirtualCameraManager::~FConcertVirtualCameraManager()
+{
+	IConcertClientPtr ConcertClient = IConcertModule::Get().GetClientInstance();
+	if (ConcertClient.IsValid())
+	{
+		TSharedPtr<IConcertClientSession> ConcertClientSession = ConcertClient->GetCurrentSession();
+		if (ConcertClientSession.IsValid())
+		{
+			UnregisterConcertSyncHandlers(ConcertClientSession.ToSharedRef());
+		}
+
+		ConcertClient->OnSessionStartup().Remove(OnSessionStartupHandle);
+		OnSessionStartupHandle.Reset();
+
+		ConcertClient->OnSessionShutdown().Remove(OnSessionShutdownHandle);
+		OnSessionShutdownHandle.Reset();
+	}
+}
+
+
+bool FConcertVirtualCameraManager::GetLatestCameraEventData(FConcertVirtualCameraCameraEvent& OutCameraEvent) const
+{
+	OutCameraEvent = LatestCameraEventData;
+	return bIsLatestCameraEventDataValid;
+}
+
+
+void FConcertVirtualCameraManager::SendCameraEventData(const FConcertVirtualCameraCameraEvent& InCameraEvent)
+{
+	TSharedPtr<IConcertClientSession> Session = WeakSession.Pin();
+	if (Session.IsValid())
+	{
+		TArray<FGuid> ClientIds = Session->GetSessionClientEndpointIds();
+		Session->SendCustomEvent(InCameraEvent, ClientIds, EConcertMessageFlags::None);
+	}
+}
+
+
+void FConcertVirtualCameraManager::RegisterConcertSyncHandlers(TSharedRef<IConcertClientSession> InSession)
+{
+	// Hold onto the session so we can trigger events
+	WeakSession = InSession;
+
+	// Register our events
+	InSession->RegisterCustomEventHandler<FConcertVirtualCameraCameraEvent>(this, &FConcertVirtualCameraManager::HandleCameraEventData);
+}
+
+
+void FConcertVirtualCameraManager::UnregisterConcertSyncHandlers(TSharedRef<IConcertClientSession> InSession)
+{
+	// Unregister our events and explicitly reset the session ptr
+	TSharedPtr<IConcertClientSession> Session = WeakSession.Pin();
+	if (Session.IsValid())
+	{
+		Session->UnregisterCustomEventHandler<FConcertVirtualCameraCameraEvent>();
+	}
+
+	WeakSession.Reset();
+}
+
+
+void FConcertVirtualCameraManager::HandleCameraEventData(const FConcertSessionContext& InEventContext, const FConcertVirtualCameraCameraEvent& InEvent)
+{
+	LatestCameraEventData = InEvent;
+	bIsLatestCameraEventDataValid = true;
+}
+
+
+#else //#if VIRTUALCAMERA_WITH_CONCERT
+
+
+bool FConcertVirtualCameraManager::GetLatestCameraEventData(FConcertVirtualCameraCameraEvent& OutCameraEvent) const
+{
+	return false;
+}
+
+
+void FConcertVirtualCameraManager::SendCameraEventData(const FConcertVirtualCameraCameraEvent& InCameraEvent)
+{
+
+}
+
+
+#endif //#if VIRTUALCAMERA_WITH_CONCERT
diff --git a/Engine/Plugins/Experimental/VirtualCamera/Source/VirtualCamera/Private/ConcertVirtualCamera.h b/Engine/Plugins/Experimental/VirtualCamera/Source/VirtualCamera/Private/ConcertVirtualCamera.h
new file mode 100644
--- /dev/null
+++ b/Engine/Plugins/Experimental/VirtualCamera/Source/VirtualCamera/Private/ConcertVirtualCamera.h
@@ -0,0 +1,127 @@
+// Copyright 1998-2019 Epic Games, Inc. All Rights Reserved.
+
+#pragma once
+
+#include "CoreMinimal.h"
+#include "CineCameraComponent.h"
+#include "VirtualCameraPlayerControllerBase.h"
+#include "UObject/Object.h"
+#include "UObject/ObjectMacros.h"
+#include "ConcertVirtualCamera.generated.h"
+
+struct FConcertSessionContext;
+class IConcertClientSession;
+class UCineCameraComponent;
+
+/**
+ * Can't use FCameraFocusSettings since it use a reference to an actor
+ * The camera will always be in Manual, and will transfer the CurrentFocusDistance.
+ */
+USTRUCT()
+struct FConcertVirtualCameraCameraFocusData
+{
+public:
+	GENERATED_BODY()
+
+	FConcertVirtualCameraCameraFocusData();
+	FConcertVirtualCameraCameraFocusData(const UCineCameraComponent* InCineCameraComponent);
+
+	FCameraFocusSettings ToCameraFocusSettings() const;
+
+	UPROPERTY()
+	float ManualFocusDistance;
+	UPROPERTY()
+	float FocusSmoothingInterpSpeed;
+	UPROPERTY()
+	uint8 bSmoothFocusChanges : 1;
+};
+
+
+/**
+ *
+ */
+USTRUCT()
+struct FConcertVirtualCameraCameraEvent
+{
+public:
+	GENERATED_BODY()
+
+	FConcertVirtualCameraCameraEvent();
+
+	/** Controller settings */
+	UPROPERTY()
+	ETrackerInputSource InputSource;
+
+	/** Camera transform */
+	UPROPERTY()
+	FVector CameraActorLocation;
+	UPROPERTY()
+	FRotator CameraActorRotation;
+	UPROPERTY()
+	FVector CameraComponentLocation;
+	UPROPERTY()
+	FRotator CameraComponentRotation;
+
+	/** Camera settings */
+	UPROPERTY()
+	float CurrentAperture;
+	UPROPERTY()
+	float CurrentFocalLength; 
+	UPROPERTY()
+	FConcertVirtualCameraCameraFocusData FocusSettings;
+	UPROPERTY()
+	FCameraLensSettings LensSettings;
+	UPROPERTY()
+	FCameraFilmbackSettings FilmbackSettings;
+};
+
+
+
+#if VIRTUALCAMERA_WITH_CONCERT // Concert is only available in development mode
+
+/**
+ *
+ */
+class FConcertVirtualCameraManager
+{
+public:
+	FConcertVirtualCameraManager();
+	FConcertVirtualCameraManager(const FConcertVirtualCameraManager&) = delete;
+	FConcertVirtualCameraManager& operator=(const FConcertVirtualCameraManager&) = delete;
+	~FConcertVirtualCameraManager();
+
+public:
+	bool GetLatestCameraEventData(FConcertVirtualCameraCameraEvent& OutCameraEvent) const;
+	void SendCameraEventData(const FConcertVirtualCameraCameraEvent& InCameraEvent);
+
+private:
+	void RegisterConcertSyncHandlers(TSharedRef<IConcertClientSession> InSession);
+	void UnregisterConcertSyncHandlers(TSharedRef<IConcertClientSession> InSession);
+
+	void HandleCameraEventData(const FConcertSessionContext& InEventContext, const FConcertVirtualCameraCameraEvent& InEvent);
+
+private:
+	/** Latest event data */
+	FConcertVirtualCameraCameraEvent LatestCameraEventData;
+	bool bIsLatestCameraEventDataValid;
+
+	/** Delegate handle for a the callback when a session starts up */
+	FDelegateHandle OnSessionStartupHandle;
+
+	/** Delegate handle for a the callback when a session shuts down */
+	FDelegateHandle OnSessionShutdownHandle;
+
+	/** Weak pointer to the client session with which to send events. May be null or stale. */
+	TWeakPtr<IConcertClientSession> WeakSession;
+};
+
+#else //VIRTUALCAMERA_WITH_CONCERT
+
+class FConcertVirtualCameraManager
+{
+public:
+	bool GetLatestCameraEventData(FConcertVirtualCameraCameraEvent& OutCameraEvent) const;
+	void SendCameraEventData(const FConcertVirtualCameraCameraEvent& InCameraEvent);
+};
+
+#endif //VIRTUALCAMERA_WITH_CONCERT
diff --git a/Engine/Plugins/Experimental/VirtualCamera/Source/VirtualCamera/Private/VirtualCamera.cpp b/Engine/Plugins/Experimental/VirtualCamera/Source/VirtualCamera/Private/VirtualCamera.cpp
--- a/Engine/Plugins/Experimental/VirtualCamera/Source/VirtualCamera/Private/VirtualCamera.cpp
+++ b/Engine/Plugins/Experimental/VirtualCamera/Source/VirtualCamera/Private/VirtualCamera.cpp
@@ -2,19 +2,42 @@
 
 #include "VirtualCamera.h"
 
-#define LOCTEXT_NAMESPACE "FVirtualCameraModule"
 
-void FVirtualCameraModule::StartupModule()
+#include "ConcertVirtualCamera.h"
+
+DEFINE_LOG_CATEGORY(LogVirtualCamera);
+
+const IVirtualCameraModule& IVirtualCameraModule::Get()
 {
-	// This code will execute after your module is loaded into memory; the exact timing is specified in the .uplugin file per-module
+	static const FName ModuleName = TEXT("VirtualCamera");
+	return FModuleManager::Get().GetModuleChecked<IVirtualCameraModule>(ModuleName);
 }
 
-void FVirtualCameraModule::ShutdownModule()
+
+/**
+ *
+ */
+class FVirtualCameraModuleImpl : public IVirtualCameraModule
 {
-	// This function may be called during shutdown to clean up your module.  For modules that support dynamic reloading,
-	// we call this function before unloading the module.
-}
+public:
+	virtual FConcertVirtualCameraManager* GetConcertVirtualCameraManager() const override
+	{
+		return ConcertManager.Get();
+	}
+
+private:
+	virtual void StartupModule() override
+	{
+		ConcertManager = MakeUnique<FConcertVirtualCameraManager>();
+	}
+
+	virtual void ShutdownModule() override
+	{
+		ConcertManager.Reset();
+	}
+
+	TUniquePtr<FConcertVirtualCameraManager> ConcertManager;
+};
+
 
-#undef LOCTEXT_NAMESPACE
-	
-IMPLEMENT_MODULE(FVirtualCameraModule, VirtualCamera)
\ No newline at end of file
+IMPLEMENT_MODULE(FVirtualCameraModuleImpl, VirtualCamera)
\ No newline at end of file
diff --git a/Engine/Plugins/Experimental/VirtualCamera/Source/VirtualCamera/Private/VirtualCamera.h b/Engine/Plugins/Experimental/VirtualCamera/Source/VirtualCamera/Private/VirtualCamera.h
new file mode 100644
--- /dev/null
+++ b/Engine/Plugins/Experimental/VirtualCamera/Source/VirtualCamera/Private/VirtualCamera.h
@@ -0,0 +1,23 @@
+// Copyright 1998-2019 Epic Games, Inc. All Rights Reserved.
+
+#pragma once
+
+#include "CoreMinimal.h"
+#include "Modules/ModuleInterface.h"
+#include "Modules/ModuleManager.h"
+
+DECLARE_LOG_CATEGORY_EXTERN(LogVirtualCamera, Log, Log);
+
+
+class FConcertVirtualCameraManager;
+
+/**
+ *
+ */
+class IVirtualCameraModule : public IModuleInterface
+{
+public:
+	static const IVirtualCameraModule& Get();
+
+	virtual FConcertVirtualCameraManager* GetConcertVirtualCameraManager() const = 0;
+};
diff --git a/Engine/Plugins/Experimental/VirtualCamera/Source/VirtualCamera/Private/VirtualCameraMovementComponent.cpp b/Engine/Plugins/Experimental/VirtualCamera/Source/VirtualCamera/Private/VirtualCameraMovementComponent.cpp
--- a/Engine/Plugins/Experimental/VirtualCamera/Source/VirtualCamera/Private/VirtualCameraMovementComponent.cpp
+++ b/Engine/Plugins/Experimental/VirtualCamera/Source/VirtualCamera/Private/VirtualCameraMovementComponent.cpp
@@ -43,9 +43,9 @@ void UVirtualCameraMovementComponent::AddInputVectorFromController(FVector World
 	}
 
 	WorldVector *= AxisSettings[MovementScaleAxis].MovementScale;
-	ApplyLocationLocks(WorldVector);
 
-	TargetLocation += WorldVector;
+	// For Controller movement, move the RootUpdatedComponent
+	FromControllerTargetLocation += WorldVector;
 }
 
 void UVirtualCameraMovementComponent::ProcessMovementInput(const FVector& TrackerLocation, const FRotator& TrackerRotation)
@@ -64,8 +64,20 @@ void UVirtualCameraMovementComponent::ProcessMovementInput(const FVector& Tracke
 	TargetRotation = TrackerRotation - GetRotationOffset() + GetOwner()->GetActorRotation();
 	PreviousTrackerRotation = TrackerRotation;
 
-	FHitResult OutResult = FHitResult();
-	SafeMoveUpdatedComponent(GetStabilizedDeltaLocation(), GetStabilizedRotation(), false, OutResult);
+	{
+		FHitResult OutResult = FHitResult();
+		SafeMoveUpdatedComponent(GetStabilizedDeltaLocation(), GetStabilizedRotation(), false, OutResult);
+	}
+
+	// Apply the location change for the RootUpdatedComponent
+	if (RootUpdatedComponent)
+	{
+		TGuardValue<USceneComponent*> UpdateComponentGuard(UpdatedComponent, RootUpdatedComponent);
+		TGuardValue<UPrimitiveComponent*> UpdatePrimitiveGuard(UpdatedPrimitive, RootUpdatedPrimitive);
+		FHitResult OutResult = FHitResult();
+		SafeMoveUpdatedComponent(FromControllerTargetLocation, FQuat::Identity, false, OutResult);
+		FromControllerTargetLocation = FVector::ZeroVector;
+	}
 }
 
 bool UVirtualCameraMovementComponent::ToggleAxisLock(const EVirtualCameraAxis AxisToToggle)
@@ -122,7 +134,6 @@ void UVirtualCameraMovementComponent::ResetCameraOffsetsToTracker()
 	UpdatedComponent->SetRelativeRotation(TargetRotation);
 
 	// Clear all locks
-	// ToDo: Do we need to clear freeze states here as well?
 	for (auto Iterator = AxisSettings.CreateIterator(); Iterator; ++Iterator)
 	{
 		Iterator.Value().SetIsLocked(false);
@@ -163,7 +174,7 @@ void UVirtualCameraMovementComponent::OnMoveUp(const float InValue)
 	if (bUseGlobalBoom)
 	{
 		FVector InputVector = FVector(0.f, 0.f, 1.f);
-		AddInputVector(InputVector * InValue);
+		AddInputVectorFromController(InputVector * InValue, EVirtualCameraAxis::LocationZ);
 	}
 	else
 	{
@@ -183,6 +194,18 @@ void UVirtualCameraMovementComponent::Teleport(const FTransform& TargetTransform
 
 	UpdatedComponent->AddLocalOffset(DeltaOffset);
 	TargetLocation += DeltaOffset;
+
+	if (RootUpdatedComponent)
+	{
+		RootUpdatedComponent->SetRelativeLocation(FVector::ZeroVector);
+		FromControllerTargetLocation = FVector::ZeroVector;
+	}
+}
+
+void UVirtualCameraMovementComponent::SetRootComponent(USceneComponent* InFromController)
+{
+	RootUpdatedComponent = InFromController;
+	RootUpdatedPrimitive = Cast<UPrimitiveComponent>(RootUpdatedComponent);
 }
 
 void UVirtualCameraMovementComponent::ApplyLocationLocks(const FVector& InVector)
diff --git a/Engine/Plugins/Experimental/VirtualCamera/Source/VirtualCamera/Private/VirtualCameraPawnBase.cpp b/Engine/Plugins/Experimental/VirtualCamera/Source/VirtualCamera/Private/VirtualCameraPawnBase.cpp
--- a/Engine/Plugins/Experimental/VirtualCamera/Source/VirtualCamera/Private/VirtualCameraPawnBase.cpp
+++ b/Engine/Plugins/Experimental/VirtualCamera/Source/VirtualCamera/Private/VirtualCameraPawnBase.cpp
@@ -23,6 +23,7 @@ AVirtualCameraPawnBase::AVirtualCameraPawnBase(const FObjectInitializer& ObjectI
 
 	MovementComponent = CreateDefaultSubobject<UVirtualCameraMovementComponent>("Movement Component");
 	MovementComponent->UpdatedComponent = CineCamera;
+	MovementComponent->SetRootComponent(DefaultSceneRoot);
 
 	// Set Default Variables
 	HomeWaypointName = "";
diff --git a/Engine/Plugins/Experimental/VirtualCamera/Source/VirtualCamera/Private/VirtualCameraPlayerControllerBase.cpp b/Engine/Plugins/Experimental/VirtualCamera/Source/VirtualCamera/Private/VirtualCameraPlayerControllerBase.cpp
--- a/Engine/Plugins/Experimental/VirtualCamera/Source/VirtualCamera/Private/VirtualCameraPlayerControllerBase.cpp
+++ b/Engine/Plugins/Experimental/VirtualCamera/Source/VirtualCamera/Private/VirtualCameraPlayerControllerBase.cpp
@@ -2,6 +2,7 @@
 
 #include "VirtualCameraPlayerControllerBase.h"
 
+#include "ConcertVirtualCamera.h"
 #include "CineCameraActor.h"
 #include "Engine/Engine.h"
 #include "Engine/World.h"
@@ -11,23 +12,14 @@
 #include "RemoteSession/Channels/RemoteSessionInputChannel.h"
 #include "RemoteSession/Channels/RemoteSessionFrameBufferChannel.h"
 #include "RemoteSession/Channels/RemoteSessionXRTrackingChannel.h"
-#include "SteamVRFunctionLibrary.h"
+#include "VirtualCamera.h"
+#include "VPGameMode.h"
+#include "VPRootActor.h"
 
 #if WITH_EDITOR
 #include "Recorder/TakeRecorderBlueprintLibrary.h"
-#include "ScopedTransaction.h"
 #endif
 
-FVirtualCameraPlayerControllerMultiUserOptions::FVirtualCameraPlayerControllerMultiUserOptions()
-	: bUpdateTargetCameraProperties(true)
-	, bUpdateTargetCameraTransform(true)
-	, bUseTransactionActor(false)
-	, bCreateTransactionTargetCameraProperties(false)
-	, bCreateTransactionTargetCameraTransform(false)
-{
-
-}
-
 
 AVirtualCameraPlayerControllerBase::AVirtualCameraPlayerControllerBase(const FObjectInitializer& ObjectInitializer)
 {
@@ -53,15 +45,52 @@ AVirtualCameraPlayerControllerBase::AVirtualCameraPlayerControllerBase(const FOb
 	CurrentFocusMethod = EVirtualCameraFocusMethod::Manual;
 }
 
+void AVirtualCameraPlayerControllerBase::OnPossess(APawn* InPawn)
+{
+	bCachedShouldUpdateTargetCameraTransform = ShouldUpdateTargetCameraTransform();
+	bCachedIsVirtualCameraControlledByRemoteSession = IsVirtualCameraControlledByRemoteSession();
+
+	Super::OnPossess(InPawn);
+}
+
 void AVirtualCameraPlayerControllerBase::BeginPlay()
 {
+	// Find the Root Actor
+	RootActor = nullptr;
+	AVPGameMode* VPGameMode = Cast<AVPGameMode>(GetWorld()->GetAuthGameMode());
+	if (VPGameMode)
+	{
+		RootActor = VPGameMode->GetRootActor();
+	}
+	else
+	{
+		UE_LOG(LogVirtualCamera, Warning, TEXT("The Game Mode is not a VPGameMode."));
+	}
+
+	if (RootActor == nullptr)
+	{
+		UE_LOG(LogVirtualCamera, Warning, TEXT("There is no VP Root Actor in the scene. A CineCameraActor will be spawned. Multi user functionalities may suffer."));
+	}
+
 	// Make a default sequence playback controller
 	LevelSequencePlaybackController = NewObject<ULevelSequencePlaybackController>(this);
-
 	if (LevelSequencePlaybackController)
 	{
-		//Always spawn the target camera that the level sequence will use as a target
-		TargetCameraActor = GetWorld()->SpawnActor<ACineCameraActor>(TargetCameraActorClass);
+		// Get the Cine Camera actor used by that root actor
+		if (RootActor)
+		{
+			TargetCameraActor = RootActor->GetCineCameraActor();
+			if (TargetCameraActor == nullptr)
+			{
+				UE_LOG(LogVirtualCamera, Warning, TEXT("The Root Actor doesn't have a cinematic camera."));
+			}
+		}
+
+		if (TargetCameraActor == nullptr)
+		{
+			// Spawn the target camera that the level sequence will use as a target
+			TargetCameraActor = GetWorld()->SpawnActor<ACineCameraActor>(TargetCameraActorClass);
+		}
 
 		// bLockToHmd is set to true by default. Remove it to prevent unwanted movement from XR system.
 		TargetCameraActor->GetCameraComponent()->bLockToHmd = false;
@@ -70,17 +99,25 @@ void AVirtualCameraPlayerControllerBase::BeginPlay()
 		LevelSequencePlaybackController->OnRecordEnabledStateChanged.BindUObject(this, &AVirtualCameraPlayerControllerBase::HandleRecordEnabledStateChange);
 	}
 
-	if (IRemoteSessionModule* RemoteSession = FModuleManager::LoadModulePtr<IRemoteSessionModule>("RemoteSession"))
+	// Is this Controller a RemoteSession controller
+	if (bCachedIsVirtualCameraControlledByRemoteSession)
 	{
-        TMap<FString, ERemoteSessionChannelMode> RequiredChannels;
-        RequiredChannels.Add(FRemoteSessionFrameBufferChannel::StaticType(),ERemoteSessionChannelMode::Write);
-        RequiredChannels.Add(FRemoteSessionInputChannel::StaticType(),ERemoteSessionChannelMode::Read);
-        RequiredChannels.Add(FRemoteSessionXRTrackingChannel::StaticType(),ERemoteSessionChannelMode::Read);
-        
-        RemoteSession->SetSupportedChannels(RequiredChannels);
-		RemoteSession->InitHost();
+		if (IRemoteSessionModule* RemoteSession = FModuleManager::LoadModulePtr<IRemoteSessionModule>("RemoteSession"))
+		{
+			TMap<FString, ERemoteSessionChannelMode> RequiredChannels;
+			RequiredChannels.Add(FRemoteSessionFrameBufferChannel::StaticType(), ERemoteSessionChannelMode::Write);
+			RequiredChannels.Add(FRemoteSessionInputChannel::StaticType(), ERemoteSessionChannelMode::Read);
+			RequiredChannels.Add(FRemoteSessionXRTrackingChannel::StaticType(), ERemoteSessionChannelMode::Read);
+
+			RemoteSession->SetSupportedChannels(RequiredChannels);
+			RemoteSession->InitHost();
+		}
 	}
-	
+	else
+	{
+		SetViewTargetWithBlend(TargetCameraActor);
+	}
+
 	if (UVirtualCameraCineCameraComponent* CineCamera = GetVirtualCameraCineCameraComponent())
 	{
 		// Need to make sure we don't let ARKit control camera completely
@@ -95,23 +132,16 @@ void AVirtualCameraPlayerControllerBase::BeginPlay()
 	// Initialize the view of the camera with offsets taken into account
 	UpdatePawnWithTrackerData();
 	ResetCameraOffsetsToTracker();
-	
-	Super::BeginPlay();
 
-	// To do a transition with Concert. Normally it should be done via the Component.
-	if (MultiUserOptions.bUseTransactionActor)
-	{
-		FActorSpawnParameters SpawnParam;
-		SpawnParam.bNoFail = true;
-		TransactionActor = GetWorld()->SpawnActor<AVirtualCameraPlayerControllerTransaction>(AVirtualCameraPlayerControllerTransaction::StaticClass(), SpawnParam);
-	}
+	Super::BeginPlay();
 }
 
 void AVirtualCameraPlayerControllerBase::Tick(float DeltaSeconds)
 {
 	Super::Tick(DeltaSeconds);
-	
-	if (AVirtualCameraPawnBase* VCPawn = GetVirtualCameraPawn())
+
+	AVirtualCameraPawnBase* VCPawn = GetVirtualCameraPawn();
+	if (VCPawn)
 	{
 		UpdatePawnWithTrackerData();
 
@@ -135,7 +165,7 @@ void AVirtualCameraPlayerControllerBase::Tick(float DeltaSeconds)
 
 	if (LevelSequencePlaybackController)
 	{
-		PilotTargetedCamera(VCCamera);
+		PilotTargetedCamera(VCPawn, VCCamera);
 	}
 }
 
@@ -164,12 +194,22 @@ ACineCameraActor* AVirtualCameraPlayerControllerBase::GetTargetCamera()
 	return TargetCameraActor;
 }
 
+void AVirtualCameraPlayerControllerBase::SetInputSource(ETrackerInputSource InInputSource)
+{
+	if (InputSource != InInputSource)
+	{
+		InputSource = InInputSource;
+
+		bCachedShouldUpdateTargetCameraTransform = ShouldUpdateTargetCameraTransform();
+	}
+}
+
 FString AVirtualCameraPlayerControllerBase::GetDistanceInDesiredUnits(const float InputDistance, const EUnit ConversionUnit) const
 {
 	// Checks if the specified conversion unit is a unit of distance, since this function assumes conversion from Unreal Units
 	if (!FUnitConversion::IsUnitOfType(ConversionUnit, EUnitType::Distance))
 	{
-		UE_LOG(LogActor, Warning, TEXT("Conversion unit selected is not a unit of distance."))
+		UE_LOG(LogVirtualCamera, Warning, TEXT("Conversion unit selected is not a unit of distance."))
 		return FString();
 	}
 
@@ -226,7 +266,6 @@ static const FName RemoteSessionTrackingSystemName(TEXT("RemoteSessionXRTracking
 
 bool AVirtualCameraPlayerControllerBase::GetCurrentTrackerLocationAndRotation(FVector& OutTrackerLocation, FRotator& OutTrackerRotation)
 {
-	TArray<int32> TrackedDeviceIDs;
 	FQuat ARKitQuaternion;
 
 	switch (InputSource)
@@ -253,23 +292,12 @@ bool AVirtualCameraPlayerControllerBase::GetCurrentTrackerLocationAndRotation(FV
 			}
 			break;
 
-		case ETrackerInputSource::Vive:
-			USteamVRFunctionLibrary::GetValidTrackedDeviceIds(ESteamVRTrackedDeviceType::Other, TrackedDeviceIDs);
-
-			// ToDo: Add filtering in event that there are multiple valid trackers available
-			if (TrackedDeviceIDs.Num() > 0)
-			{
-				USteamVRFunctionLibrary::GetTrackedDevicePositionAndOrientation(TrackedDeviceIDs[0], OutTrackerLocation, OutTrackerRotation);
-				return true;
-			}
-			break;
-
 		case ETrackerInputSource::Custom:
 			GetCustomTrackerLocationAndRotation(OutTrackerLocation, OutTrackerRotation);
 			return true;
 
 		default:
-			UE_LOG(LogActor, Warning, TEXT("Selected tracker source is not yet supported"))
+			UE_LOG(LogVirtualCamera, Warning, TEXT("Selected tracker source is not yet supported"))
 			break;
 	}
 
@@ -394,59 +422,70 @@ UVirtualCameraMovementComponent* AVirtualCameraPlayerControllerBase::GetVirtualC
 	return nullptr;
 }
 
-void AVirtualCameraPlayerControllerBase::PilotTargetedCamera(UVirtualCameraCineCameraComponent* CameraToFollow)
+void AVirtualCameraPlayerControllerBase::PilotTargetedCamera(AVirtualCameraPawnBase* PawnToFollow, UVirtualCameraCineCameraComponent* CameraToFollow)
 {
-	if (!TargetCameraActor || !CameraToFollow)
+	if (TargetCameraActor == nullptr || PawnToFollow == nullptr || CameraToFollow == nullptr)
 	{
 		return;
 	}
 
-	if (MultiUserOptions.bUpdateTargetCameraTransform)
+	UCineCameraComponent* TargetCameraComponent = TargetCameraActor->GetCineCameraComponent();
+	if (TargetCameraComponent == nullptr)
 	{
-#if WITH_EDITOR
-		FScopedTransaction Transaction(NSLOCTEXT("VirtualCamera", "SetVCamTransform", "Set VCam Transform"), MultiUserOptions.bCreateTransactionTargetCameraTransform);
-		TargetCameraActor->Modify();
-#endif //WITH_EDITOR
-
-		TargetCameraActor->SetActorLocationAndRotation(CameraToFollow->GetComponentLocation(), CameraToFollow->GetComponentRotation().Quaternion());
+		return;
 	}
 
-	if (MultiUserOptions.bUpdateTargetCameraProperties)
+	bool bAssignValues = false;
+	FConcertVirtualCameraCameraEvent CameraEvent;
+	if (bCachedIsVirtualCameraControlledByRemoteSession)
 	{
-		if (UCineCameraComponent* TargetComponent = TargetCameraActor->GetCineCameraComponent())
-		{
-#if WITH_EDITOR
-			FScopedTransaction Transaction(NSLOCTEXT("VirtualCamera", "SetVCamSettings", "Set VCam Settings"), MultiUserOptions.bCreateTransactionTargetCameraProperties);
-			TargetComponent->Modify();
-			if (TransactionActor)
-			{
-				TransactionActor->Modify();
-			}
-#endif //WITH_EDITOR
+		CameraEvent.InputSource = InputSource;
 
-			TargetComponent->CurrentAperture = CameraToFollow->CurrentAperture;
-			TargetComponent->CurrentFocalLength = CameraToFollow->CurrentFocalLength;
-			TargetComponent->FocusSettings = CameraToFollow->FocusSettings;
-			TargetComponent->LensSettings = CameraToFollow->LensSettings;
-			TargetComponent->FilmbackSettings = CameraToFollow->DesiredFilmbackSettings;
+		CameraEvent.CameraActorLocation = PawnToFollow->GetActorLocation();
+		CameraEvent.CameraActorRotation = PawnToFollow->GetActorRotation();
+		CameraEvent.CameraComponentLocation = CameraToFollow->RelativeLocation;
+		CameraEvent.CameraComponentRotation = CameraToFollow->RelativeRotation;
 
-			if (TransactionActor)
-			{
-				TransactionActor->Aperture = CameraToFollow->CurrentAperture;
-				TransactionActor->FocalLength = CameraToFollow->CurrentFocalLength;
-				TransactionActor->FilmbackSettings = CameraToFollow->DesiredFilmbackSettings;
-				TransactionActor->FocusSettings = CameraToFollow->FocusSettings;
-				TransactionActor->LensSettings = CameraToFollow->LensSettings;
-			}
+		CameraEvent.CurrentAperture = CameraToFollow->CurrentAperture;
+		CameraEvent.CurrentFocalLength = CameraToFollow->CurrentFocalLength;
+		CameraEvent.FocusSettings = FConcertVirtualCameraCameraFocusData(CameraToFollow);
+		CameraEvent.LensSettings = CameraToFollow->LensSettings;
+		CameraEvent.FilmbackSettings = CameraToFollow->DesiredFilmbackSettings;
+
+		IVirtualCameraModule::Get().GetConcertVirtualCameraManager()->SendCameraEventData(CameraEvent);
+		bAssignValues = true;
+	}
+	else
+	{
+		bAssignValues = IVirtualCameraModule::Get().GetConcertVirtualCameraManager()->GetLatestCameraEventData(CameraEvent);
+
+		PawnToFollow->SetActorLocationAndRotation(CameraEvent.CameraActorLocation, CameraEvent.CameraActorRotation);
+		if (bCachedShouldUpdateTargetCameraTransform)
+		{
+			CameraToFollow->SetRelativeLocationAndRotation(CameraEvent.CameraComponentLocation, CameraEvent.CameraComponentRotation);
+		}
+		CameraToFollow->CurrentAperture = CameraEvent.CurrentAperture;
+		CameraToFollow->CurrentFocalLength = CameraEvent.CurrentFocalLength;
+		CameraToFollow->FocusSettings = CameraEvent.FocusSettings.ToCameraFocusSettings();
+		CameraToFollow->LensSettings = CameraEvent.LensSettings;
+		CameraToFollow->DesiredFilmbackSettings = CameraEvent.FilmbackSettings;
+
+		if (CameraEvent.InputSource != InputSource)
+		{
+			SetInputSource(CameraEvent.InputSource);
 		}
 	}
-	else if (TransactionActor)
+
+	// Copy the info to the camera target camera
+	if (bAssignValues)
 	{
-		CameraToFollow->CurrentAperture = TransactionActor->Aperture;
-		CameraToFollow->CurrentFocalLength = TransactionActor->FocalLength;
-		CameraToFollow->DesiredFilmbackSettings = TransactionActor->FilmbackSettings;
-		CameraToFollow->FocusSettings = TransactionActor->FocusSettings;
-		CameraToFollow->LensSettings = TransactionActor->LensSettings;
+		TargetCameraActor->SetActorLocationAndRotation(CameraEvent.CameraActorLocation, CameraEvent.CameraActorRotation);
+		TargetCameraComponent->SetRelativeLocationAndRotation(CameraToFollow->RelativeLocation, CameraToFollow->RelativeRotation);
+		TargetCameraComponent->CurrentAperture = CameraToFollow->CurrentAperture;
+		TargetCameraComponent->CurrentFocalLength = CameraToFollow->CurrentFocalLength;
+		TargetCameraComponent->FocusSettings = CameraToFollow->FocusSettings;
+		TargetCameraComponent->LensSettings = CameraToFollow->LensSettings;
+		TargetCameraComponent->FilmbackSettings = CameraToFollow->FilmbackSettings;
 	}
 }
 
@@ -564,15 +603,16 @@ void AVirtualCameraPlayerControllerBase::UpdatePawnWithTrackerData()
 	{
 		FVector TrackerLocation = FVector::ZeroVector;
 		FRotator TrackerRotation = FRotator::ZeroRotator;
-		GetCurrentTrackerLocationAndRotation(TrackerLocation, TrackerRotation);
-
-		// Apply tracker offset to tracker; convert everything to transforms to make sure motions are calculated in the right order
-		FTransform TrackerRaw = FTransform(TrackerRotation, TrackerLocation);
-		FTransform AdjustedTracker = TrackerPostOffset.AsTransform() * TrackerRaw * TrackerPreOffset.AsTransform();
-		TrackerRotation = AdjustedTracker.Rotator();
-		TrackerLocation = AdjustedTracker.GetLocation();
+		if (GetCurrentTrackerLocationAndRotation(TrackerLocation, TrackerRotation))
+		{
+			// Apply tracker offset to tracker; convert everything to transforms to make sure motions are calculated in the right order
+			FTransform TrackerRaw = FTransform(TrackerRotation, TrackerLocation);
+			FTransform AdjustedTracker = TrackerPostOffset.AsTransform() * TrackerRaw * TrackerPreOffset.AsTransform();
+			TrackerRotation = AdjustedTracker.Rotator();
+			TrackerLocation = AdjustedTracker.GetLocation();
 
-		VCPawn->ProcessMovementInput(TrackerLocation, TrackerRotation);
+			VCPawn->ProcessMovementInput(TrackerLocation, TrackerRotation);
+		}
 	}
 }
 
@@ -609,6 +649,16 @@ bool AVirtualCameraPlayerControllerBase::IsLocationWithinMatte(const FVector Loc
 	return false;
 }
 
+bool AVirtualCameraPlayerControllerBase::IsVirtualCameraControlledByRemoteSession_Implementation() const
+{
+	return true;
+}
+
+bool AVirtualCameraPlayerControllerBase::ShouldUpdateTargetCameraTransform_Implementation() const
+{
+	return InputSource == ETrackerInputSource::ARKit || IsVirtualCameraControlledByRemoteSession();
+}
+
 /***** UI Interface Functions *****/
 
 void AVirtualCameraPlayerControllerBase::AddBlendableToCamera(TScriptInterface<IBlendableInterface> BlendableToAdd, float Weight)
diff --git a/Engine/Plugins/Experimental/VirtualCamera/Source/VirtualCamera/Public/VirtualCamera.h b/Engine/Plugins/Experimental/VirtualCamera/Source/VirtualCamera/Public/VirtualCamera.h
deleted file mode 100644
diff --git a/Engine/Plugins/Experimental/VirtualCamera/Source/VirtualCamera/Public/VirtualCameraMovementComponent.h b/Engine/Plugins/Experimental/VirtualCamera/Source/VirtualCamera/Public/VirtualCameraMovementComponent.h
--- a/Engine/Plugins/Experimental/VirtualCamera/Source/VirtualCamera/Public/VirtualCameraMovementComponent.h
+++ b/Engine/Plugins/Experimental/VirtualCamera/Source/VirtualCamera/Public/VirtualCameraMovementComponent.h
@@ -171,6 +171,14 @@ public:
 	 */
 	void SetUseGlobalBoom(bool bShouldUseGlobalBoom) { bUseGlobalBoom = bShouldUseGlobalBoom; }
 
+
+	/**
+	 * When the input come from a physical controller, move this component instead.
+	 * This allow a proper replication when we are in a multi user setting.
+	 * This component doesn't respect the freeze and lock axis.
+	 */
+	void SetRootComponent(USceneComponent* FromController);
+
 private:
 	/** The cached axis for locking, only used if bLockRelativeToFirstLockAxis is true */
 	FQuat CachedLockingAxis;
@@ -181,12 +189,29 @@ private:
 	/** Tracks the current target location of the camera for stabilization */
 	FVector TargetLocation;
 
+	/** Tracks the current target location of the camera that came from a controller */
+	FVector FromControllerTargetLocation;
+
 	/** Tracks the previous tracker rotation to determine how much we should rotate */
 	FRotator PreviousTrackerRotation;
 
 	/** Tracks the current target rotation of the camera for stabilization */
 	FRotator TargetRotation;
 
+	/**
+	 * The component we move and update when the input is coming from the controller.
+	 * If this is null UpdatedComponent will be used instead.
+	 * @see UpdatedComponent
+	 */
+	UPROPERTY(Transient, DuplicateTransient)
+	USceneComponent* RootUpdatedComponent;
+
+	/**
+	 * RootUpdatedComponent, cast as a UPrimitiveComponent. May be invalid if RootUpdatedComponent was null or not a UPrimitiveComponent.
+	 */
+	UPROPERTY(Transient, DuplicateTransient)
+	UPrimitiveComponent* RootUpdatedPrimitive;
+
 	/**
 	 * Applies relative axis filtering based on locks for location.
 	 * @param InVector - The input vector that will be used to apply locking offsets
diff --git a/Engine/Plugins/Experimental/VirtualCamera/Source/VirtualCamera/Public/VirtualCameraPlayerControllerBase.h b/Engine/Plugins/Experimental/VirtualCamera/Source/VirtualCamera/Public/VirtualCameraPlayerControllerBase.h
--- a/Engine/Plugins/Experimental/VirtualCamera/Source/VirtualCamera/Public/VirtualCameraPlayerControllerBase.h
+++ b/Engine/Plugins/Experimental/VirtualCamera/Source/VirtualCamera/Public/VirtualCameraPlayerControllerBase.h
@@ -4,17 +4,18 @@
 
 #include "GameFramework/PlayerController.h"
 
-#include "Misc/FrameNumber.h"
-#include "Misc/FrameRate.h"
-#include "Misc/Timecode.h"
+#include "CineCameraActor.h"
 #include "ILiveLinkClient.h"
 #include "InputCore.h"
 #include "LevelSequencePlaybackController.h"
 #include "Math/UnitConversion.h"
+#include "Misc/FrameNumber.h"
+#include "Misc/FrameRate.h"
+#include "Misc/Timecode.h"
 #include "VirtualCameraPawnBase.h"
-
 #include "VirtualCameraPlayerControllerBase.generated.h"
 
+class AVPRootActor;
 
 UENUM(BlueprintType)
 enum class ETrackerInputSource : uint8
@@ -26,10 +27,7 @@ enum class ETrackerInputSource : uint8
 	Custom,
 
 	/* Livelink plugin tracker */
-	LiveLink,
-
-	/* Vive tracker */
-	Vive,
+	LiveLink
 };
 
 UENUM(BlueprintType)
@@ -78,66 +76,13 @@ public:
 	};
 };
 
-USTRUCT(BlueprintType)
-struct FVirtualCameraPlayerControllerMultiUserOptions
-{
-public:
-	GENERATED_BODY()
-
-	FVirtualCameraPlayerControllerMultiUserOptions();
-	
-	/** Should the controller update the target camera settings from the pawn's camera. */
-	UPROPERTY(EditAnywhere, BlueprintReadWrite, Category = "Camera")
-	bool bUpdateTargetCameraProperties;
-
-	/** Should the controller update the camera transform from the pawn's camera. */
-	UPROPERTY(EditAnywhere, BlueprintReadWrite, Category = "Camera")
-	bool bUpdateTargetCameraTransform;
-
-	/** Should the controller update use the Transition actor to transfer the camera settings. */
-	UPROPERTY(EditAnywhere, BlueprintReadWrite, Category = "Camera")
-	bool bUseTransactionActor;
-
-	/**
-	 * Create a transaction when a camera settings has been modified.
-	 * @note Transaction can only be created in the editor.
-	 */
-	UPROPERTY(EditAnywhere, BlueprintReadWrite, Category = "Camera")
-	bool bCreateTransactionTargetCameraProperties;
-
-	/**
-	 * Create a transaction when a camera transform has been modified.
-	 * @note Transaction can only be created in the editor.
-	 */
-	UPROPERTY(EditAnywhere, BlueprintReadWrite, Category = "Camera")
-	bool bCreateTransactionTargetCameraTransform;
-};
-
-// To do a transition with Concert. Normally it should be done via the Component.
-UCLASS()
-class VIRTUALCAMERA_API AVirtualCameraPlayerControllerTransaction : public AActor
-{
-	GENERATED_BODY()
-
-public:
-	UPROPERTY(EditAnywhere, Category = "Transaction")
-	float Aperture;
-	UPROPERTY(EditAnywhere, Category = "Transaction")
-	float FocalLength;
-	UPROPERTY(EditAnywhere, Category = "Transaction")
-	FCameraFilmbackSettings FilmbackSettings;
-	UPROPERTY(EditAnywhere, Category = "Transaction")
-	FCameraFocusSettings FocusSettings;
-	UPROPERTY(EditAnywhere, Category = "Transaction")
-	FCameraLensSettings LensSettings;
-};
-
-UCLASS()
+UCLASS(Abstract)
 class VIRTUALCAMERA_API AVirtualCameraPlayerControllerBase : public APlayerController
 {
 	GENERATED_UCLASS_BODY()
 
 public:
+	virtual void OnPossess(APawn* InPawn) override;
 	virtual void BeginPlay() override;
 	virtual void Tick(float DeltaSeconds) override;
 	virtual void SetupInputComponent() override;
@@ -145,28 +90,25 @@ public:
 	FScriptDelegate OnStop;
 
 	/** Allows user to select which tracker input should be used */
-	UPROPERTY(EditDefaultsOnly, BlueprintReadWrite, Category = "Camera Settings")
+	UPROPERTY(EditDefaultsOnly, BlueprintReadWrite, BlueprintSetter="SetInputSource", Category = "Camera Settings")
 	ETrackerInputSource InputSource;
 
 	/** Controller for level sequence playback */
-	UPROPERTY(EditAnywhere, BlueprintReadWrite, Category = "Camera Settings")
+	UPROPERTY(EditAnywhere, BlueprintReadOnly, Category = "Camera Settings")
 	FName LiveLinkTargetName;
 
 	/** Offset applied to calculated location before tracker transform is added */
-	UPROPERTY(EditAnywhere, BlueprintReadWrite, Category = "Camera Settings")
+	UPROPERTY(EditAnywhere, BlueprintReadOnly, Category = "Camera Settings")
 	FTrackingOffset TrackerPreOffset;
 
 	/** Offset applied to calculated location after tracker transform is added */
-	UPROPERTY(EditAnywhere, BlueprintReadWrite, Category = "Camera Settings")
+	UPROPERTY(EditAnywhere, BlueprintReadOnly, Category = "Camera Settings")
 	FTrackingOffset TrackerPostOffset;
 
 	/** Class of CameraActor to spawn to allow user to use their own customized camera */
-	UPROPERTY(EditAnywhere, BlueprintReadWrite, NoClear, Category = "Camera Settings")
+	UPROPERTY(EditAnywhere, BlueprintReadOnly, NoClear, Category = "Camera Settings")
 	TSubclassOf<ACineCameraActor> TargetCameraActorClass;
 
-	UPROPERTY(EditAnywhere, BlueprintReadWrite, NoClear, Category = "Transaction")
-	FVirtualCameraPlayerControllerMultiUserOptions MultiUserOptions;
-
 	/** Array of any properties that should be recorded */
 	UPROPERTY(EditAnywhere, Category = "Recording")
 	TArray<FName> RequiredSequencerRecorderCameraSettings;
@@ -174,8 +116,8 @@ public:
 	UPROPERTY(BlueprintAssignable, Category = "Movement")
 	FVirtualCameraResetOffsetsDelegate OnOffsetReset;
 
-	UPROPERTY(transient)
-	AVirtualCameraPlayerControllerTransaction* TransactionActor;
+	UPROPERTY(transient, BlueprintReadOnly, Category = "VirtualCamera")
+	AVPRootActor* RootActor;
 
 	/**
 	 * Overridable function to allow user to get tracker data from blueprints.
@@ -217,6 +159,12 @@ public:
 	UFUNCTION(BlueprintCallable, Category = "VirtualCamera")
 	ACineCameraActor* GetTargetCamera();
 
+	/**
+	 * Set the input source
+	 */
+	UFUNCTION(BlueprintSetter, Category = "Camera Settings")
+	void SetInputSource(ETrackerInputSource InInputSource);
+
 	/**
 	 * Blueprint event for when the focus method is changed.
 	 * @param NewFocuMethod - The focus method
@@ -259,10 +207,19 @@ protected:
 	UPROPERTY(Transient)
 	ULevelSequencePlaybackController* LevelSequencePlaybackController;
 	
-	/** Target camera that is spawned on begin play for the sequence controller */
+	/** Target camera that is spawned or possessed on begin play for the sequence controller */
 	UPROPERTY(Transient)
 	ACineCameraActor* TargetCameraActor;
 
+	/** Cached value for IsVirtualCameraControlledByRemoteSession() */
+	UPROPERTY(Transient, BlueprintReadOnly, Category = "VirtualCamera")
+	bool bCachedIsVirtualCameraControlledByRemoteSession;
+
+	/** Cached value for IsVirtualCameraControlledByRemoteSession() */
+	UPROPERTY(Transient, BlueprintReadOnly, Category = "VirtualCamera")
+	bool bCachedShouldUpdateTargetCameraTransform;
+	
+
 protected:
 
 	/**
@@ -307,7 +264,7 @@ protected:
 	/**
 	 * Pilot the controlled camera during recording, copying over settings from the pawn.
 	 */
-	void PilotTargetedCamera(UVirtualCameraCineCameraComponent* CameraToFollow);
+	void PilotTargetedCamera(AVirtualCameraPawnBase* PawnToFollow, UVirtualCameraCineCameraComponent* CameraToFollow);
 
 	/**
 	 * Override of InputTouch, used to handle touch and hold events.
@@ -367,6 +324,18 @@ protected:
 	UFUNCTION()
 	void BroadcastOffsetReset() { OnOffsetReset.Broadcast(); }
 
+	/*
+	 * Is this machine should display the Virtual Camera UI and establish a connection with the remote session app.
+	 */
+	UFUNCTION(BlueprintNativeEvent, Category = "VirtualCamera")
+	bool IsVirtualCameraControlledByRemoteSession() const;
+
+	/*
+	 * In multi user session, how should we update the information across different sessions.
+	 */
+	UFUNCTION(BlueprintNativeEvent, Category = "VirtualCamera")
+	bool ShouldUpdateTargetCameraTransform() const;
+
 	/***** UI Interface *****/
 public:
 
diff --git a/Engine/Plugins/Experimental/VirtualCamera/Source/VirtualCamera/VirtualCamera.Build.cs b/Engine/Plugins/Experimental/VirtualCamera/Source/VirtualCamera/VirtualCamera.Build.cs
--- a/Engine/Plugins/Experimental/VirtualCamera/Source/VirtualCamera/VirtualCamera.Build.cs
+++ b/Engine/Plugins/Experimental/VirtualCamera/Source/VirtualCamera/VirtualCamera.Build.cs
@@ -7,33 +7,48 @@ public class VirtualCamera : ModuleRules
 	public VirtualCamera(ReadOnlyTargetRules Target) : base(Target)
 	{
 		PCHUsage = ModuleRules.PCHUsageMode.UseExplicitOrSharedPCHs;
-			
+
 		PublicDependencyModuleNames.AddRange(
 			new string[]
 			{
-                "AugmentedReality",
-                "CinematicCamera",
-                "Core",
-                "CoreUObject",
-                "Engine",
-                "HeadMountedDisplay",
-                "InputCore",
-                "LevelSequence",
-                "LiveLinkInterface",
-                "MovieScene",
-                "RemoteSession",
-                "SteamVR",
-                "TimeManagement"
-            }
-            );
+				"AugmentedReality",
+				"CinematicCamera",
+				"Core",
+				"CoreUObject",
+				"Engine",
+				"InputCore",
+				"LevelSequence",
+				"LiveLinkInterface",
+				"MovieScene",
+				"RemoteSession",
+				"TimeManagement",
+				"VPUtilities",
+			}
+		);
+
+		if (Target.bBuildDeveloperTools)
+		{
+			PrivateDefinitions.Add("VIRTUALCAMERA_WITH_CONCERT=1");
+			PrivateDependencyModuleNames.AddRange(
+				new string[]
+				{
+					"Concert",
+				}
+			);
+		}
+		else
+		{
+			PrivateDefinitions.Add("VIRTUALCAMERA_WITH_CONCERT=0");
+		}
+
 
-        if (Target.bBuildEditor == true)
-        {
+		if (Target.bBuildEditor == true)
+		{
 			PublicDependencyModuleNames.Add("LevelSequenceEditor");
 			PublicDependencyModuleNames.Add("Sequencer");
 			PublicDependencyModuleNames.Add("SlateCore");
 			PublicDependencyModuleNames.Add("TakeRecorder");
 			PrivateDependencyModuleNames.Add("UnrealEd");
-        }
+		}
 	}
 }
diff --git a/Engine/Plugins/Experimental/VirtualCamera/VirtualCamera.uplugin b/Engine/Plugins/Experimental/VirtualCamera/VirtualCamera.uplugin
--- a/Engine/Plugins/Experimental/VirtualCamera/VirtualCamera.uplugin
+++ b/Engine/Plugins/Experimental/VirtualCamera/VirtualCamera.uplugin
@@ -38,11 +38,11 @@
 			"Enabled": true
 		},
 		{
-			"Name": "SteamVR",
+			"Name": "Takes",
 			"Enabled": true
 		},
 		{
-			"Name": "Takes",
+			"Name": "VirtualProductionUtilities",
 			"Enabled": true
 		}
 	]
diff --git a/Engine/Plugins/Experimental/VirtualProductionUtilities/Source/VPUtilities/Private/VPGameMode.cpp b/Engine/Plugins/Experimental/VirtualProductionUtilities/Source/VPUtilities/Private/VPGameMode.cpp
new file mode 100644
--- /dev/null
+++ b/Engine/Plugins/Experimental/VirtualProductionUtilities/Source/VPUtilities/Private/VPGameMode.cpp
@@ -0,0 +1,42 @@
+// Copyright 1998-2019 Epic Games, Inc. All Rights Reserved.
+
+#include "VPGameMode.h"
+
+#include "EngineUtils.h"
+#include "GameFramework/Actor.h"
+#include "GameFramework/Controller.h"
+#include "GameFramework/WorldSettings.h"
+#include "VPRootActor.h"
+
+
+AActor* AVPGameMode::ChoosePlayerStart_Implementation(AController* Player)
+{
+	RootActor = nullptr;
+
+	AWorldSettings* WorldSetting = GetWorld() ? GetWorld()->GetWorldSettings() : nullptr;
+	if (WorldSetting)
+	{
+		UVPWorldAssetUserData* UserData = CastChecked<UVPWorldAssetUserData>(WorldSetting->GetAssetUserDataOfClass(UVPWorldAssetUserData::StaticClass()), ECastCheckedType::NullAllowed);
+		if (UserData)
+		{
+			RootActor = UserData->LastSelectedRootActor.Get();
+		}
+	}
+
+	if (RootActor == nullptr)
+	{
+		for (TActorIterator<AVPRootActor> It(GetWorld()); It; ++It)
+		{
+			RootActor = *It;
+			break;
+		}
+	}
+
+	AActor* Result = RootActor;
+	if (Result == nullptr)
+	{
+		Result = Super::ChoosePlayerStart_Implementation(Player);
+	}
+
+	return Result;
+}
diff --git a/Engine/Plugins/Experimental/VirtualProductionUtilities/Source/VPUtilities/Private/VPRootActor.cpp b/Engine/Plugins/Experimental/VirtualProductionUtilities/Source/VPUtilities/Private/VPRootActor.cpp
new file mode 100644
--- /dev/null
+++ b/Engine/Plugins/Experimental/VirtualProductionUtilities/Source/VPUtilities/Private/VPRootActor.cpp
@@ -0,0 +1,234 @@
+// Copyright 1998-2019 Epic Games, Inc. All Rights Reserved.
+
+#include "VPRootActor.h"
+
+#include "CineCameraActor.h"
+#include "Components/ArrowComponent.h"
+#include "Components/BillboardComponent.h"
+#include "Components/StaticMeshComponent.h"
+#include "Engine/Texture2D.h"
+#include "GameFramework/PlayerStart.h"
+#include "UObject/ConstructorHelpers.h"
+
+#if WITH_EDITOR
+#include "Engine/LevelStreaming.h"
+#include "Engine/Level.h"
+#include "Engine/World.h"
+#include "LevelUtils.h"
+#include "Logging/MessageLog.h"
+#include "Logging/TokenizedMessage.h"
+#include "Misc/UObjectToken.h"
+#endif
+
+#define LOCTEXT_NAMESPACE "VirtualProductionUtilities"
+
+
+AVPRootActor::AVPRootActor(const FObjectInitializer& ObjectInitializer)
+	: Super(ObjectInitializer)
+	, bAreComponentsVisible(true)
+#if WITH_EDITOR
+	, bReentrantPostEditMove(false)
+#endif
+{
+	RootComponent = CreateDefaultSubobject<USceneComponent>(TEXT("Root"));
+
+	RealWorldSceneRepresentation = CreateDefaultSubobject<UStaticMeshComponent>(TEXT("RealWorldRepresentation"));
+	if (RealWorldSceneRepresentation)
+	{
+		RealWorldSceneRepresentation->bHiddenInGame = false;
+		RealWorldSceneRepresentation->SetupAttachment(RootComponent);
+	}
+
+#if WITH_EDITOR
+	struct FConstructorStatics
+	{
+		FName ID_Sprite;
+		FText NAME_Sprite;
+		FColor SceneBaseColor;
+		FVector SceneBaseSize;
+		FConstructorStatics()
+			: ID_Sprite(TEXT("VPRootActor"))
+			, NAME_Sprite(LOCTEXT("RootSpriteInfo", "VP Root Actor"))
+			, SceneBaseColor(100, 255, 255, 255)
+			, SceneBaseSize(600.0f, 600.0f, 400.0f)
+		{
+		}
+	};
+	static FConstructorStatics ConstructorStatics;
+
+	SpriteComponent = CreateEditorOnlyDefaultSubobject<UBillboardComponent>(TEXT("Sprite"));
+	if (SpriteComponent)
+	{
+		ConstructorHelpers::FObjectFinderOptional<UTexture2D> RootTextureObject = TEXT("/VirtualProductionUtilities/Icons/S_VPRootActor");
+
+		SpriteComponent->Sprite = RootTextureObject.Get();
+		SpriteComponent->RelativeScale3D = FVector(0.5f, 0.5f, 0.5f);
+		SpriteComponent->bHiddenInGame = false;
+		SpriteComponent->SpriteInfo.Category = ConstructorStatics.ID_Sprite;
+		SpriteComponent->SpriteInfo.DisplayName = ConstructorStatics.NAME_Sprite;
+		SpriteComponent->bIsScreenSizeScaled = true;
+		SpriteComponent->SetupAttachment(RootComponent);
+	}
+
+	ArrowComponent = CreateEditorOnlyDefaultSubobject<UArrowComponent>(TEXT("Arrow"));
+	if (ArrowComponent)
+	{
+		ArrowComponent->ArrowColor = FColor(150, 200, 255);
+		ArrowComponent->bTreatAsASprite = true;
+		ArrowComponent->bHiddenInGame = false;
+		ArrowComponent->SpriteInfo.Category = ConstructorStatics.ID_Sprite;
+		ArrowComponent->SpriteInfo.DisplayName = ConstructorStatics.NAME_Sprite;
+		ArrowComponent->bIsScreenSizeScaled = true;
+		ArrowComponent->SetupAttachment(RootComponent);
+	}
+#endif
+}
+
+
+void AVPRootActor::BeginPlay()
+{
+	if (IsRunningGame())
+	{
+		bAreComponentsVisible = false;
+	}
+	SetComponentsVisibility(bAreComponentsVisible);
+
+	Super::BeginPlay();
+}
+
+
+namespace VPRootActor
+{
+	template<class T>
+	T* GetFirstActor_Implementation(const AVPRootActor* RootActor)
+	{
+		if (USceneComponent* SceneComponent = RootActor->GetRootComponent())
+		{
+			TArray<USceneComponent*> ChildComponents;
+			SceneComponent->GetChildrenComponents(true, ChildComponents);
+			for (USceneComponent* ChildComponent : ChildComponents)
+			{
+				if (AActor* ChildActor = ChildComponent->GetOwner())
+				{
+					if (T* CineCameraActor = Cast<T>(ChildActor))
+					{
+						return CineCameraActor;
+					}
+				}
+			}
+		}
+
+		return nullptr;
+	}
+}
+
+
+ACineCameraActor* AVPRootActor::GetCineCameraActor_Implementation() const
+{
+	if (CinematicCamera == nullptr)
+	{
+		return CinematicCamera;
+	}
+
+	return VPRootActor::GetFirstActor_Implementation<ACineCameraActor>(this);
+}
+
+
+#if WITH_EDITOR
+void AVPRootActor::ToggleComponentsVisibility()
+{
+	Modify();
+	bAreComponentsVisible = !bAreComponentsVisible;
+	SetComponentsVisibility(bAreComponentsVisible);
+}
+
+
+void AVPRootActor::MoveLevelToRootActor()
+{
+	ULevel* OwningLevel = GetLevel();
+	UWorld* OwningWorld = GetWorld();
+	ULevelStreaming* OwningStreamingLevel = nullptr;
+	if (OwningWorld && OwningLevel)
+	{
+		ULevelStreaming *const * FoundItem = OwningWorld->GetStreamingLevels().FindByPredicate([OwningLevel](ULevelStreaming* InItem) { return InItem && InItem->GetLoadedLevel() == OwningLevel; });
+		if (FoundItem)
+		{
+			OwningStreamingLevel = *FoundItem;
+		}
+	}
+
+	if (OwningStreamingLevel)
+	{
+		FVector ActorLocation = GetActorLocation();
+		FRotator ActorRotation = GetActorRotation();
+		ActorRotation.Pitch = 0;
+		ActorRotation.Roll = 0;
+		FLevelUtils::SetEditorTransform(OwningStreamingLevel, FTransform(ActorRotation, ActorLocation));
+		SetActorLocationAndRotation(ActorLocation, ActorRotation, false);
+	}
+}
+
+
+void AVPRootActor::CheckForErrors()
+{
+	Super::CheckForErrors();
+
+	if (GetCineCameraActor() == nullptr)
+	{
+		FMessageLog("MapCheck").Warning()
+			->AddToken(FUObjectToken::Create(this))
+			->AddToken(FTextToken::Create(LOCTEXT("MapCheck_Message_CinematicCameraNull", "The Root Actor has doesn't have a Cinematic Camera Actor set.")));
+	}
+}
+
+
+void AVPRootActor::PostEditChangeChainProperty(struct FPropertyChangedChainEvent& PropertyChangedEvent)
+{
+	Super::PostEditChangeChainProperty(PropertyChangedEvent);
+
+	if (bMoveLevelWithActor && !bReentrantPostEditMove)
+	{
+		if (PropertyChangedEvent.Property->GetFName() == GET_MEMBER_NAME_CHECKED(USceneComponent, RelativeLocation)
+			|| PropertyChangedEvent.Property->GetFName() == GET_MEMBER_NAME_CHECKED(USceneComponent, RelativeRotation))
+		{
+			TGuardValue<bool> Tmp(bReentrantPostEditMove, true);
+			MoveLevelToRootActor();
+		}
+	}
+}
+
+
+void AVPRootActor::PostEditMove(bool bFinished)
+{
+	Super::PostEditMove(bFinished);
+
+	if (bMoveLevelWithActor && !bReentrantPostEditMove && bFinished)
+	{
+		TGuardValue<bool> Tmp(bReentrantPostEditMove, true);
+		MoveLevelToRootActor();
+	}
+}
+#endif //WITH_EDITOR
+
+
+void AVPRootActor::SetComponentsVisibility(bool bVisible)
+{
+	if (RealWorldSceneRepresentation)
+	{
+		RealWorldSceneRepresentation->SetVisibility(bVisible, true);
+	}
+
+#if WITH_EDITOR
+	if (SpriteComponent)
+	{
+		SpriteComponent->SetVisibility(bVisible, true);
+	}
+	if (ArrowComponent)
+	{
+		ArrowComponent->SetVisibility(bVisible, true);
+	}
+#endif
+}
+
+
+#undef LOCTEXT_NAMESPACE
diff --git a/Engine/Plugins/Experimental/VirtualProductionUtilities/Source/VPUtilities/Public/VPGameMode.h b/Engine/Plugins/Experimental/VirtualProductionUtilities/Source/VPUtilities/Public/VPGameMode.h
new file mode 100644
--- /dev/null
+++ b/Engine/Plugins/Experimental/VirtualProductionUtilities/Source/VPUtilities/Public/VPGameMode.h
@@ -0,0 +1,33 @@
+// Copyright 1998-2018 Epic Games, Inc. All Rights Reserved.
+
+#pragma once
+
+#include "CoreMinimal.h"
+#include "GameFramework/GameMode.h"
+
+#include "VPGameMode.generated.h"
+
+
+class AActor;
+class AController;
+class AVPRootActor;
+
+/**
+ * 
+ */
+UCLASS(abstract)
+class VPUTILITIES_API AVPGameMode : public AGameMode
+{
+	GENERATED_BODY()
+
+private:
+	UPROPERTY(transient, BlueprintGetter = "GetRootActor", Category = "Virtual Production")
+	AVPRootActor* RootActor;
+
+public:
+	UFUNCTION(BlueprintGetter, Category = "Virtual Production")
+	AVPRootActor* GetRootActor() const { return RootActor; }
+
+public:
+	virtual AActor* ChoosePlayerStart_Implementation(AController* Player) override;
+};
diff --git a/Engine/Plugins/Experimental/VirtualProductionUtilities/Source/VPUtilities/Public/VPRootActor.h b/Engine/Plugins/Experimental/VirtualProductionUtilities/Source/VPUtilities/Public/VPRootActor.h
new file mode 100644
--- /dev/null
+++ b/Engine/Plugins/Experimental/VirtualProductionUtilities/Source/VPUtilities/Public/VPRootActor.h
@@ -0,0 +1,95 @@
+// Copyright 1998-2019 Epic Games, Inc. All Rights Reserved.
+
+#pragma once
+
+#include "CoreMinimal.h"
+
+#include "Components/BoxComponent.h"
+#include "Engine/AssetUserData.h"
+#include "GameFramework/Actor.h"
+#include "VPRootActor.generated.h"
+
+
+class ACineCameraActor;
+class UArrowComponent;
+class UBillboardComponent;
+class UMeshComponent;
+
+
+/**
+ * AVPRootActor
+ */
+UCLASS(meta = (DisplayName = "VP Root Actor"), hidecategories=(Actor, Collision, Input, LOD, Replication))
+class VPUTILITIES_API AVPRootActor : public AActor
+{
+	GENERATED_BODY()
+
+public:
+	AVPRootActor(const FObjectInitializer& ObjectInitializer);
+
+public:
+	/** Actor used to show the size of the scene. */
+	UPROPERTY(EditDefaultsOnly, Category = "Virtual Production")
+	UMeshComponent* RealWorldSceneRepresentation;
+
+	/** Actor used to show the size of the scene. */
+	UPROPERTY(EditAnywhere, Category = "Virtual Production")
+	ACineCameraActor* CinematicCamera;
+
+private:
+	UPROPERTY()
+	bool bAreComponentsVisible;
+
+#if WITH_EDITORONLY_DATA
+protected:
+	/** Billboard used to see the scene in the editor */
+	UPROPERTY()
+	UBillboardComponent* SpriteComponent;
+
+	/** Arrow used to see the orientation of the scene in the editor */
+	UPROPERTY()
+	UArrowComponent* ArrowComponent;
+
+public:
+	UPROPERTY(EditAnywhere, Category = "Virtual Production")
+	bool bMoveLevelWithActor;
+
+	bool bReentrantPostEditMove;
+#endif
+
+public:
+	/** Get the current camera used by the virtual production. */
+	UFUNCTION(BlueprintNativeEvent, Category = "Virtual Production")
+	ACineCameraActor* GetCineCameraActor() const;
+
+	virtual void BeginPlay() override;
+
+#if WITH_EDITOR
+	UFUNCTION(BlueprintCallable, CallInEditor, Category = "Virtual Production")
+	void ToggleComponentsVisibility();
+
+	UFUNCTION(BlueprintCallable, CallInEditor, Category = "Virtual Production")
+	void MoveLevelToRootActor();
+
+	virtual void CheckForErrors() override;
+	virtual void PostEditMove(bool bFinished) override;
+	virtual void PostEditChangeChainProperty(struct FPropertyChangedChainEvent& PropertyChangedEvent) override;
+#endif
+
+private:
+	void SetComponentsVisibility(bool bVisible);
+};
+
+
+/**
+ * UVPWorldAssetUserData
+ */
+UCLASS(MinimalAPI)
+class UVPWorldAssetUserData : public UAssetUserData
+{
+	GENERATED_BODY()
+
+public:
+	UPROPERTY(EditAnywhere, Category = "Virtual Production")
+	TLazyObjectPtr<AVPRootActor> LastSelectedRootActor;
+};
