Update the virtual camera position when it's setup to be in a MU session but with no openned session.

diff --git a/Engine/Plugins/Experimental/VirtualCamera/Source/VirtualCamera/Private/VirtualCameraPlayerControllerBase.cpp b/Engine/Plugins/Experimental/VirtualCamera/Source/VirtualCamera/Private/VirtualCameraPlayerControllerBase.cpp
--- a/Engine/Plugins/Experimental/VirtualCamera/Source/VirtualCamera/Private/VirtualCameraPlayerControllerBase.cpp
+++ b/Engine/Plugins/Experimental/VirtualCamera/Source/VirtualCamera/Private/VirtualCameraPlayerControllerBase.cpp
@@ -43,6 +43,10 @@ AVirtualCameraPlayerControllerBase::AVirtualCameraPlayerControllerBase(const FOb
 	TouchInputState = ETouchInputState::BlueprintDefined;
 	PreviousTouchInput = TouchInputState;
 	CurrentFocusMethod = EVirtualCameraFocusMethod::Manual;
+
+	// Cached values
+	bCachedShouldUpdateTargetCameraTransform = true;
+	bCachedIsVirtualCameraControlledByRemoteSession = true;
 }
 
 void AVirtualCameraPlayerControllerBase::OnPossess(APawn* InPawn)
@@ -458,21 +462,23 @@ void AVirtualCameraPlayerControllerBase::PilotTargetedCamera(AVirtualCameraPawnB
 	else
 	{
 		bAssignValues = IVirtualCameraModule::Get().GetConcertVirtualCameraManager()->GetLatestCameraEventData(CameraEvent);
-
-		PawnToFollow->SetActorLocationAndRotation(CameraEvent.CameraActorLocation, CameraEvent.CameraActorRotation);
-		if (bCachedShouldUpdateTargetCameraTransform)
+		if (bAssignValues)
 		{
-			CameraToFollow->SetRelativeLocationAndRotation(CameraEvent.CameraComponentLocation, CameraEvent.CameraComponentRotation);
-		}
-		CameraToFollow->CurrentAperture = CameraEvent.CurrentAperture;
-		CameraToFollow->CurrentFocalLength = CameraEvent.CurrentFocalLength;
-		CameraToFollow->FocusSettings = CameraEvent.FocusSettings.ToCameraFocusSettings();
-		CameraToFollow->LensSettings = CameraEvent.LensSettings;
-		CameraToFollow->DesiredFilmbackSettings = CameraEvent.FilmbackSettings;
+			PawnToFollow->SetActorLocationAndRotation(CameraEvent.CameraActorLocation, CameraEvent.CameraActorRotation);
+			if (bCachedShouldUpdateTargetCameraTransform)
+			{
+				CameraToFollow->SetRelativeLocationAndRotation(CameraEvent.CameraComponentLocation, CameraEvent.CameraComponentRotation);
+			}
+			CameraToFollow->CurrentAperture = CameraEvent.CurrentAperture;
+			CameraToFollow->CurrentFocalLength = CameraEvent.CurrentFocalLength;
+			CameraToFollow->FocusSettings = CameraEvent.FocusSettings.ToCameraFocusSettings();
+			CameraToFollow->LensSettings = CameraEvent.LensSettings;
+			CameraToFollow->DesiredFilmbackSettings = CameraEvent.FilmbackSettings;
 
-		if (CameraEvent.InputSource != InputSource)
-		{
-			SetInputSource(CameraEvent.InputSource);
+			if (CameraEvent.InputSource != InputSource)
+			{
+				SetInputSource(CameraEvent.InputSource);
+			}
 		}
 	}
 
@@ -480,13 +486,13 @@ void AVirtualCameraPlayerControllerBase::PilotTargetedCamera(AVirtualCameraPawnB
 	if (bAssignValues)
 	{
 		TargetCameraActor->SetActorLocationAndRotation(CameraEvent.CameraActorLocation, CameraEvent.CameraActorRotation);
-		TargetCameraComponent->SetRelativeLocationAndRotation(CameraToFollow->RelativeLocation, CameraToFollow->RelativeRotation);
 		TargetCameraComponent->CurrentAperture = CameraToFollow->CurrentAperture;
 		TargetCameraComponent->CurrentFocalLength = CameraToFollow->CurrentFocalLength;
 		TargetCameraComponent->FocusSettings = CameraToFollow->FocusSettings;
 		TargetCameraComponent->LensSettings = CameraToFollow->LensSettings;
 		TargetCameraComponent->FilmbackSettings = CameraToFollow->FilmbackSettings;
 	}
+	TargetCameraComponent->SetRelativeLocationAndRotation(CameraToFollow->RelativeLocation, CameraToFollow->RelativeRotation);
 }
 
 bool AVirtualCameraPlayerControllerBase::InputTouch(uint32 Handle, ETouchType::Type Type, const FVector2D & TouchLocation, float Force, FDateTime DeviceTimestamp, uint32 TouchpadIndex)
