Update test map to test VCam with LiveLink.
Add deprecation path when using an old asset to control the VCam camera via LiveLink

diff --git a/Engine/Plugins/Experimental/VirtualCamera/Source/VirtualCamera/Private/VirtualCameraPlayerControllerBase.cpp b/Engine/Plugins/Experimental/VirtualCamera/Source/VirtualCamera/Private/VirtualCameraPlayerControllerBase.cpp
--- a/Engine/Plugins/Experimental/VirtualCamera/Source/VirtualCamera/Private/VirtualCameraPlayerControllerBase.cpp
+++ b/Engine/Plugins/Experimental/VirtualCamera/Source/VirtualCamera/Private/VirtualCameraPlayerControllerBase.cpp
@@ -13,6 +13,8 @@
 #include "RemoteSession/Channels/RemoteSessionInputChannel.h"
 #include "RemoteSession/Channels/RemoteSessionFrameBufferChannel.h"
 #include "RemoteSession/Channels/RemoteSessionXRTrackingChannel.h"
+#include "Roles/LiveLinkAnimationRole.h"
+#include "Roles/LiveLinkAnimationTypes.h"
 #include "Roles/LiveLinkTransformRole.h"
 #include "Roles/LiveLinkTransformTypes.h"
 #include "VirtualCamera.h"
@@ -293,13 +295,25 @@ bool AVirtualCameraPlayerControllerBase::GetCurrentTrackerLocationAndRotation(FV
 				if (LiveLinkClient->EvaluateFrame_AnyThread(LiveLinkTargetName, ULiveLinkTransformRole::StaticClass(), EvaluateData))
 				{
 					FLiveLinkTransformFrameData* TransformFrameData = EvaluateData.FrameData.Cast<FLiveLinkTransformFrameData>();
-					if (TransformFrameData)
+					check(TransformFrameData);
+
+					OutTrackerLocation = TransformFrameData->Transform.GetLocation();
+					OutTrackerRotation = TransformFrameData->Transform.GetRotation().Rotator();
+
+					return true;
+				}
+				else if (LiveLinkClient->EvaluateFrame_AnyThread(LiveLinkTargetName, ULiveLinkAnimationRole::StaticClass(), EvaluateData))
+				{
+					FLiveLinkAnimationFrameData* AnimationFrameData = EvaluateData.FrameData.Cast<FLiveLinkAnimationFrameData>();
+					check(AnimationFrameData);
+					if (AnimationFrameData->Transforms.Num() > 0)
 					{
-						OutTrackerLocation = TransformFrameData->Transform.GetLocation();
-						OutTrackerRotation = TransformFrameData->Transform.GetRotation().Rotator();
+						OutTrackerLocation = AnimationFrameData->Transforms[0].GetLocation();
+						OutTrackerRotation = AnimationFrameData->Transforms[0].GetRotation().Rotator();
+
+						return true;
 					}
 				}
-				return true;
 			}
 			break;
 
