Missing file.

#rnx

diff --git a/Engine/Plugins/Experimental/VirtualCamera/Source/VirtualCamera/VirtualCamera.Build.cs b/Engine/Plugins/Experimental/VirtualCamera/Source/VirtualCamera/VirtualCamera.Build.cs
--- a/Engine/Plugins/Experimental/VirtualCamera/Source/VirtualCamera/VirtualCamera.Build.cs
+++ b/Engine/Plugins/Experimental/VirtualCamera/Source/VirtualCamera/VirtualCamera.Build.cs
@@ -33,7 +33,7 @@ public class VirtualCamera : ModuleRules
 			}
 		);
 
-		if (Target.bBuildDeveloperTools)
+		if (Target.Type == TargetType.Editor || Target.Type == TargetType.Program)
 		{
 			PrivateDefinitions.Add("VIRTUALCAMERA_WITH_CONCERT=1");
 			PrivateDependencyModuleNames.AddRange(
