Resolving a bad merge in VirtualCamera.
When resolving the VirtualCamera from LiveLink, use the parent bone.

diff --git a/Engine/Plugins/Experimental/VirtualCamera/Source/VirtualCamera/Private/VirtualCameraPlayerControllerBase.cpp b/Engine/Plugins/Experimental/VirtualCamera/Source/VirtualCamera/Private/VirtualCameraPlayerControllerBase.cpp
--- a/Engine/Plugins/Experimental/VirtualCamera/Source/VirtualCamera/Private/VirtualCameraPlayerControllerBase.cpp
+++ b/Engine/Plugins/Experimental/VirtualCamera/Source/VirtualCamera/Private/VirtualCameraPlayerControllerBase.cpp
@@ -13,6 +13,8 @@
 #include "RemoteSession/Channels/RemoteSessionInputChannel.h"
 #include "RemoteSession/Channels/RemoteSessionFrameBufferChannel.h"
 #include "RemoteSession/Channels/RemoteSessionXRTrackingChannel.h"
+#include "Roles/LiveLinkAnimationRole.h"
+#include "Roles/LiveLinkAnimationTypes.h"
 #include "Roles/LiveLinkTransformRole.h"
 #include "Roles/LiveLinkTransformTypes.h"
 #include "VirtualCamera.h"
@@ -293,13 +295,29 @@ bool AVirtualCameraPlayerControllerBase::GetCurrentTrackerLocationAndRotation(FV
 				if (LiveLinkClient->EvaluateFrame_AnyThread(LiveLinkTargetName, ULiveLinkTransformRole::StaticClass(), EvaluateData))
 				{
 					FLiveLinkTransformFrameData* TransformFrameData = EvaluateData.FrameData.Cast<FLiveLinkTransformFrameData>();
-					if (TransformFrameData)
+					check(TransformFrameData);
+
+					OutTrackerLocation = TransformFrameData->Transform.GetLocation();
+					OutTrackerRotation = TransformFrameData->Transform.GetRotation().Rotator();
+
+					return true;
+				}
+				else if (LiveLinkClient->EvaluateFrame_AnyThread(LiveLinkTargetName, ULiveLinkAnimationRole::StaticClass(), EvaluateData))
+				{
+					FLiveLinkAnimationFrameData* AnimationFrameData = EvaluateData.FrameData.Cast<FLiveLinkAnimationFrameData>();
+					FLiveLinkSkeletonStaticData* AnimationStaticData = EvaluateData.StaticData.Cast<FLiveLinkSkeletonStaticData>();
+					check(AnimationFrameData && AnimationStaticData);
+
+					int32 RootIndex = AnimationStaticData->FindRootBone();
+					if (AnimationFrameData->Transforms.IsValidIndex(RootIndex))
 					{
-						OutTrackerLocation = TransformFrameData->Transform.GetLocation();
-						OutTrackerRotation = TransformFrameData->Transform.GetRotation().Rotator();
+						
+						OutTrackerLocation = AnimationFrameData->Transforms[RootIndex].GetLocation();
+						OutTrackerRotation = AnimationFrameData->Transforms[RootIndex].GetRotation().Rotator();
+
+						return true;
 					}
 				}
-				return true;
 			}
 			break;
 
diff --git a/Engine/Source/Runtime/LiveLinkInterface/Public/Roles/LiveLinkAnimationTypes.h b/Engine/Source/Runtime/LiveLinkInterface/Public/Roles/LiveLinkAnimationTypes.h
--- a/Engine/Source/Runtime/LiveLinkInterface/Public/Roles/LiveLinkAnimationTypes.h
+++ b/Engine/Source/Runtime/LiveLinkInterface/Public/Roles/LiveLinkAnimationTypes.h
@@ -28,6 +28,9 @@ public:
 	//Get skeleton's parent bones array
 	const TArray<int32>& GetBoneParents() const { return BoneParents; }
 
+	// Find skeleton root bone, which is the bone with an invalid parent bone index.
+	int32 FindRootBone() const { return BoneParents.IndexOfByPredicate([](int32 BoneParent){ return BoneParent < 0; }); }
+
 public:
 
 	// Names of each bone in the skeleton
