Move the Virtual Camera Movement from the Dev-Enterprise-VCam to Dev-VirtualProduction.

diff --git a/Engine/Plugins/Experimental/VirtualCamera/Source/VirtualCamera/Private/VirtualCameraMovement.cpp b/Engine/Plugins/Experimental/VirtualCamera/Source/VirtualCamera/Private/VirtualCameraMovement.cpp
new file mode 100644
--- /dev/null
+++ b/Engine/Plugins/Experimental/VirtualCamera/Source/VirtualCamera/Private/VirtualCameraMovement.cpp
@@ -0,0 +1,66 @@
+// Copyright 1998-2019 Epic Games, Inc. All Rights Reserved.
+
+#include "VirtualCameraMovement.h"
+
+UVirtualCameraMovement::UVirtualCameraMovement()
+	: LocalTransform(FTransform::Identity)
+	, LocalAxisTransform(FTransform::Identity)
+	, CalculatedTransform(FTransform::Identity)
+	, LocationScale(FVector::OneVector)
+	, RotationScale(FRotator(1.f, 1.f, 1.f))
+	, bLocalAxisSet(false)
+{
+
+}
+
+void UVirtualCameraMovement::SetLocalTransform(const FTransform& InTransform)
+{
+	LocalTransform = InTransform;
+	Update();
+}
+
+void UVirtualCameraMovement::SetLocalAxis(const FTransform& InTransform)
+{
+	LocalAxisTransform = InTransform;
+	LocalAxisTransform.SetScale3D(FVector(1.f));
+
+	bLocalAxisSet = true;
+	Update();
+}
+
+void UVirtualCameraMovement::ResetLocalAxis()
+{
+	LocalAxisTransform = FTransform::Identity;
+	bLocalAxisSet = false;
+	Update();
+}
+
+void UVirtualCameraMovement::SetLocationScale(FVector InLocationScale)
+{
+	LocationScale = InLocationScale;
+	Update();
+}
+
+void UVirtualCameraMovement::SetRotationScale(FRotator InRotationScale)
+{
+	RotationScale = InRotationScale;
+	Update();
+}
+
+void UVirtualCameraMovement::Update()
+{
+	CalculatedTransform = LocalTransform;
+
+	FVector TransformedLocation = LocalAxisTransform.InverseTransformPosition(LocalTransform.GetLocation());
+	TransformedLocation *= LocationScale;
+	CalculatedTransform.SetLocation(LocalAxisTransform.TransformPosition(TransformedLocation));
+
+	FQuat TransformedRotation = LocalAxisTransform.InverseTransformRotation(LocalTransform.GetRotation());
+	FRotator TransformedRotator = TransformedRotation.Rotator();
+	TransformedRotator.Pitch *= RotationScale.Pitch;
+	TransformedRotator.Yaw *= RotationScale.Yaw;
+	TransformedRotator.Roll *= RotationScale.Roll;
+	TransformedRotator.Normalize();
+	LocalTransform.GetRotation().Rotator().SetClosestToMe(TransformedRotator);
+	CalculatedTransform.SetRotation(LocalAxisTransform.TransformRotation(TransformedRotator.Quaternion()));
+}
diff --git a/Engine/Plugins/Experimental/VirtualCamera/Source/VirtualCamera/Public/VirtualCameraMovement.h b/Engine/Plugins/Experimental/VirtualCamera/Source/VirtualCamera/Public/VirtualCameraMovement.h
new file mode 100644
--- /dev/null
+++ b/Engine/Plugins/Experimental/VirtualCamera/Source/VirtualCamera/Public/VirtualCameraMovement.h
@@ -0,0 +1,68 @@
+// Copyright 1998-2019 Epic Games, Inc. All Rights Reserved.
+
+#pragma once
+
+#include "CoreMinimal.h"
+#include "UObject/Object.h"
+#include "VirtualCameraMovement.generated.h"
+
+
+/**
+ * A class to handle the virtual Camera transform.
+ */
+UCLASS(BlueprintType)
+class VIRTUALCAMERA_API UVirtualCameraMovement : public UObject
+{
+	GENERATED_BODY()
+	UVirtualCameraMovement();
+
+public:
+	/** Get the unmodified local transform. */
+	UFUNCTION(BlueprintCallable, Category = "Virtual Camera")
+	FTransform GetLocalTransform() const { return LocalTransform; }
+
+	/** Set the local transform that will be modified by the scaling factor. */
+	UFUNCTION(BlueprintCallable, Category = "Virtual Camera")
+	void SetLocalTransform(const FTransform& Transform);
+
+	/** Get the modified transform. */
+	UFUNCTION(BlueprintCallable, Category = "Virtual Camera")
+	FTransform GetTransform() const { return CalculatedTransform; }
+
+	/** Set the transform from where the scaling factor will be calculated from. */
+	UFUNCTION(BlueprintCallable, Category = "Virtual Camera")
+	void SetLocalAxis(const FTransform& InTransform);
+
+	/** Remove the axis transform, all scaling factor will be calculated from the origin. */
+	UFUNCTION(BlueprintCallable, Category = "Virtual Camera")
+	void ResetLocalAxis();
+
+	/** Scale the local location of the local transform from the local axis. */
+	UFUNCTION(BlueprintCallable, Category = "Virtual Camera")
+	void SetLocationScale(FVector LocationScale);
+
+	/** Scale the local rotation of the local transform from the local axis. */
+	UFUNCTION(BlueprintCallable, Category = "Virtual Camera")
+	void SetRotationScale(FRotator RotationScale);
+
+private:
+	/** Calculate the modified transform. */
+	void Update();
+
+private:
+	/** The unmodified local transform. The transform we received from LiveLink or ARKit. */
+	FTransform LocalTransform;
+	/** Location and rotation of when the user decided to create its scale. */
+	FTransform LocalAxisTransform;
+	/** The calculated transform, that we calculated from the LocalAxis. */
+	FTransform CalculatedTransform;
+
+	/** Translation scale base on the axis point. */
+	FVector LocationScale;
+
+	/** Rotation scale base on the axis point. */
+	FRotator RotationScale;
+
+	/** Is a local axis set. */
+	bool bLocalAxisSet;
+};
