Check for invalid value returned by the ARKit for Virtual Camera.

diff --git a/Engine/Plugins/Experimental/VirtualCamera/Source/VirtualCamera/Private/VirtualCameraPlayerControllerBase.cpp b/Engine/Plugins/Experimental/VirtualCamera/Source/VirtualCamera/Private/VirtualCameraPlayerControllerBase.cpp
--- a/Engine/Plugins/Experimental/VirtualCamera/Source/VirtualCamera/Private/VirtualCameraPlayerControllerBase.cpp
+++ b/Engine/Plugins/Experimental/VirtualCamera/Source/VirtualCamera/Private/VirtualCameraPlayerControllerBase.cpp
@@ -275,16 +275,17 @@ static const FName RemoteSessionTrackingSystemName(TEXT("RemoteSessionXRTracking
 
 bool AVirtualCameraPlayerControllerBase::GetCurrentTrackerLocationAndRotation(FVector& OutTrackerLocation, FRotator& OutTrackerRotation)
 {
-	FQuat ARKitQuaternion;
+	bool bTransformSet = false;
 
 	switch (InputSource)
 	{
 		case ETrackerInputSource::ARKit:
 			if (GEngine && GEngine->XRSystem.IsValid() && GEngine->XRSystem->GetSystemName() == RemoteSessionTrackingSystemName)
 			{
+				FQuat ARKitQuaternion;
 				GEngine->XRSystem->GetCurrentPose(0, ARKitQuaternion, OutTrackerLocation);
 				OutTrackerRotation = ARKitQuaternion.Rotator();
-				return true;
+				bTransformSet = true;
 			}
 			break;
 
@@ -300,7 +301,7 @@ bool AVirtualCameraPlayerControllerBase::GetCurrentTrackerLocationAndRotation(FV
 					OutTrackerLocation = TransformFrameData->Transform.GetLocation();
 					OutTrackerRotation = TransformFrameData->Transform.GetRotation().Rotator();
 
-					return true;
+					bTransformSet = true;
 				}
 				else if (LiveLinkClient->EvaluateFrame_AnyThread(LiveLinkTargetName, ULiveLinkAnimationRole::StaticClass(), EvaluateData))
 				{
@@ -315,7 +316,7 @@ bool AVirtualCameraPlayerControllerBase::GetCurrentTrackerLocationAndRotation(FV
 						OutTrackerLocation = AnimationFrameData->Transforms[RootIndex].GetLocation();
 						OutTrackerRotation = AnimationFrameData->Transforms[RootIndex].GetRotation().Rotator();
 
-						return true;
+						bTransformSet = true;
 					}
 				}
 			}
@@ -323,15 +324,17 @@ bool AVirtualCameraPlayerControllerBase::GetCurrentTrackerLocationAndRotation(FV
 
 		case ETrackerInputSource::Custom:
 			GetCustomTrackerLocationAndRotation(OutTrackerLocation, OutTrackerRotation);
-			return true;
+			bTransformSet = true;
 
 		default:
 			UE_LOG(LogVirtualCamera, Warning, TEXT("Selected tracker source is not yet supported"))
 			break;
 	}
 
-	// Return failure status if we couldn't find device to track or device isn't supported
-	return false;
+	bTransformSet = bTransformSet && !OutTrackerLocation.ContainsNaN()) && !OutTrackerRotation.ContainsNaN();
+
+	// Return failure status if we couldn't find device to track or device isn't supported or the values contains NaN
+	return bTransformSet;
 }
 
 bool AVirtualCameraPlayerControllerBase::IsTouchInputInFocusMode()
