Fix CIS build

diff --git a/Engine/Plugins/Experimental/VirtualCamera/Source/VirtualCamera/Private/VirtualCameraPlayerControllerBase.cpp b/Engine/Plugins/Experimental/VirtualCamera/Source/VirtualCamera/Private/VirtualCameraPlayerControllerBase.cpp
--- a/Engine/Plugins/Experimental/VirtualCamera/Source/VirtualCamera/Private/VirtualCameraPlayerControllerBase.cpp
+++ b/Engine/Plugins/Experimental/VirtualCamera/Source/VirtualCamera/Private/VirtualCameraPlayerControllerBase.cpp
@@ -331,7 +331,7 @@ bool AVirtualCameraPlayerControllerBase::GetCurrentTrackerLocationAndRotation(FV
 			break;
 	}
 
-	bTransformSet = bTransformSet && !OutTrackerLocation.ContainsNaN()) && !OutTrackerRotation.ContainsNaN();
+	bTransformSet = bTransformSet && !OutTrackerLocation.ContainsNaN() && !OutTrackerRotation.ContainsNaN();
 
 	// Return failure status if we couldn't find device to track or device isn't supported or the values contains NaN
 	return bTransformSet;
