Fix for log spam when using custom tracker in VCam

diff --git a/Engine/Plugins/Experimental/VirtualCamera/Source/VirtualCamera/Private/VirtualCameraPlayerControllerBase.cpp b/Engine/Plugins/Experimental/VirtualCamera/Source/VirtualCamera/Private/VirtualCameraPlayerControllerBase.cpp
--- a/Engine/Plugins/Experimental/VirtualCamera/Source/VirtualCamera/Private/VirtualCameraPlayerControllerBase.cpp
+++ b/Engine/Plugins/Experimental/VirtualCamera/Source/VirtualCamera/Private/VirtualCameraPlayerControllerBase.cpp
@@ -325,6 +325,7 @@ bool AVirtualCameraPlayerControllerBase::GetCurrentTrackerLocationAndRotation(FV
 		case ETrackerInputSource::Custom:
 			GetCustomTrackerLocationAndRotation(OutTrackerLocation, OutTrackerRotation);
 			bTransformSet = true;
+			break;
 
 		default:
 			UE_LOG(LogVirtualCamera, Warning, TEXT("Selected tracker source is not yet supported"))
