Add interface for VirtualCamera UMG

diff --git a/Engine/Plugins/Experimental/VirtualCamera/Source/VirtualCamera/Private/LevelSequencePlaybackController.cpp b/Engine/Plugins/Experimental/VirtualCamera/Source/VirtualCamera/Private/LevelSequencePlaybackController.cpp
--- a/Engine/Plugins/Experimental/VirtualCamera/Source/VirtualCamera/Private/LevelSequencePlaybackController.cpp
+++ b/Engine/Plugins/Experimental/VirtualCamera/Source/VirtualCamera/Private/LevelSequencePlaybackController.cpp
@@ -28,14 +28,6 @@ ULevelSequencePlaybackController::ULevelSequencePlaybackController(const FObject
 
 }
 
-void ULevelSequencePlaybackController::ResumeLevelSequencePlay()
-{
-	if (ActiveLevelSequence)
-	{
-		PlayLevelSequence();
-	}
-}
-
 void ULevelSequencePlaybackController::GetLevelSequences(TArray<FLevelSequenceData>& OutLevelSequenceNames)
 {
 	OutLevelSequenceNames.Empty();
@@ -346,15 +338,6 @@ void ULevelSequencePlaybackController::ClearActiveLevelSequence()
 	}
 }
 
-void ULevelSequencePlaybackController::PlayFromBeginning()
-{
-	if (ActiveLevelSequence)
-	{
-		JumpToPlaybackPosition(GetCurrentSequencePlaybackStart());
-		PlayLevelSequence();
-	}
-}
-
 void ULevelSequencePlaybackController::PlayToEnd()
 {
 	if (ActiveLevelSequence)
diff --git a/Engine/Plugins/Experimental/VirtualCamera/Source/VirtualCamera/Private/VirtualCameraPawnBase.cpp b/Engine/Plugins/Experimental/VirtualCamera/Source/VirtualCamera/Private/VirtualCameraPawnBase.cpp
--- a/Engine/Plugins/Experimental/VirtualCamera/Source/VirtualCamera/Private/VirtualCameraPawnBase.cpp
+++ b/Engine/Plugins/Experimental/VirtualCamera/Source/VirtualCamera/Private/VirtualCameraPawnBase.cpp
@@ -413,7 +413,7 @@ void AVirtualCameraPawnBase::SaveSettings()
 	// Save waypoints
 	SaveGameInstance->Waypoints = Waypoints;
 	SaveGameInstance->HomeWaypointName = HomeWaypointName;
-	
+
 	// Save screenshots
 	SaveGameInstance->Screenshots = Screenshots;
 
@@ -541,3 +541,46 @@ FString AVirtualCameraPawnBase::LeftPadWithZeros(int32 InNumber, int32 MinNumber
 
 	return ReturnString;
 }
+
+//~ Begin IVirtualCameraPresetContainer Interface
+
+FString AVirtualCameraPawnBase::SavePreset_Implementation(const bool bSaveCameraSettings, const bool bSaveStabilization, const bool bSaveAxisLocking, const bool bSaveMotionScale)
+{
+	return SavePreset(bSaveCameraSettings, bSaveStabilization, bSaveAxisLocking, bSaveMotionScale);
+}
+
+bool AVirtualCameraPawnBase::LoadPreset_Implementation(const FString& PresetName)
+{
+	return LoadPreset(PresetName);
+}
+
+int32 AVirtualCameraPawnBase::DeletePreset_Implementation(const FString& PresetName)
+{
+	return DeletePreset(PresetName);
+}
+
+TMap<FString, FVirtualCameraSettingsPreset> AVirtualCameraPawnBase::GetSettingsPresets_Implementation()
+{
+	return GetSettingsPresets();
+}
+
+//~ End IVirtualCameraPresetContainer Interface
+
+//~ Begin IVirtualCameraOptions Interface
+
+void AVirtualCameraPawnBase::SetDesiredDistanceUnits_Implementation(const EUnit DesiredUnits)
+{
+	SetDesiredDistanceUnits(DesiredUnits);
+}
+
+EUnit AVirtualCameraPawnBase::GetDesiredDistanceUnits_Implementation()
+{
+	return GetDesiredDistanceUnits();
+}
+
+bool AVirtualCameraPawnBase::IsFocusVisualizationAllowed_Implementation()
+{
+	return IsFocusVisualizationAllowed();
+}
+
+//~ End IVirtualCameraOptions Interface
diff --git a/Engine/Plugins/Experimental/VirtualCamera/Source/VirtualCamera/Private/VirtualCameraPlayerControllerBase.cpp b/Engine/Plugins/Experimental/VirtualCamera/Source/VirtualCamera/Private/VirtualCameraPlayerControllerBase.cpp
--- a/Engine/Plugins/Experimental/VirtualCamera/Source/VirtualCamera/Private/VirtualCameraPlayerControllerBase.cpp
+++ b/Engine/Plugins/Experimental/VirtualCamera/Source/VirtualCamera/Private/VirtualCameraPlayerControllerBase.cpp
@@ -18,6 +18,7 @@
 #include "Roles/LiveLinkTransformRole.h"
 #include "Roles/LiveLinkTransformTypes.h"
 #include "VirtualCamera.h"
+#include "VirtualCameraSubsystem.h"
 #include "VPGameMode.h"
 #include "VPRootActor.h"
 
@@ -81,6 +82,12 @@ void AVirtualCameraPlayerControllerBase::BeginPlay()
 		UE_LOG(LogVirtualCamera, Warning, TEXT("There is no VP Root Actor in the scene. A CineCameraActor will be spawned. Multi user functionalities may suffer."));
 	}
 
+	// set controller implementation (aka this) onto camera subsystem
+	if (UVirtualCameraSubsystem* SubSystem = GEngine->GetEngineSubsystem<UVirtualCameraSubsystem>())
+	{
+		SubSystem->SetVirtualCameraController(this);
+	}
+
 	// Make a default sequence playback controller
 	LevelSequencePlaybackController = NewObject<ULevelSequencePlaybackController>(this);
 	if (LevelSequencePlaybackController)
@@ -1418,3 +1425,28 @@ void AVirtualCameraPlayerControllerBase::ActivateGameViewport()
 {
 	FSlateApplication::Get().ActivateGameViewport();
 }
+
+UCineCameraComponent* AVirtualCameraPlayerControllerBase::GetStreamedCameraComponent_Implementation() const
+{
+	return GetVirtualCameraCineCameraComponent();
+}
+
+UCineCameraComponent* AVirtualCameraPlayerControllerBase::GetRecordingCameraComponent_Implementation() const
+{
+	return TargetCameraActor ? TargetCameraActor->GetCineCameraComponent() : nullptr;
+}
+
+ULevelSequencePlaybackController* AVirtualCameraPlayerControllerBase::GetSequenceController_Implementation() const
+{
+	return LevelSequencePlaybackController;
+}
+
+TScriptInterface<IVirtualCameraPresetContainer> AVirtualCameraPlayerControllerBase::GetPresetContainer_Implementation() const
+{
+	return GetVirtualCameraPawn();
+}
+
+TScriptInterface<IVirtualCameraOptions> AVirtualCameraPlayerControllerBase::GetOptions_Implementation() const
+{
+	return GetVirtualCameraPawn();
+}
diff --git a/Engine/Plugins/Experimental/VirtualCamera/Source/VirtualCamera/Private/VirtualCameraSubsystem.cpp b/Engine/Plugins/Experimental/VirtualCamera/Source/VirtualCamera/Private/VirtualCameraSubsystem.cpp
new file mode 100644
--- /dev/null
+++ b/Engine/Plugins/Experimental/VirtualCamera/Source/VirtualCamera/Private/VirtualCameraSubsystem.cpp
@@ -0,0 +1,9 @@
+// Copyright 1998-2019 Epic Games, Inc. All Rights Reserved.
+
+#include "VirtualCameraSubsystem.h"
+
+void UVirtualCameraSubsystem::SetVirtualCameraController(TScriptInterface<IVirtualCameraController> VirtualCamera)
+{
+	ActiveCameraController = VirtualCamera;
+	//todo deactive the last current, initialize the new active, call back
+}
\ No newline at end of file
diff --git a/Engine/Plugins/Experimental/VirtualCamera/Source/VirtualCamera/Public/IVirtualCameraController.h b/Engine/Plugins/Experimental/VirtualCamera/Source/VirtualCamera/Public/IVirtualCameraController.h
new file mode 100644
--- /dev/null
+++ b/Engine/Plugins/Experimental/VirtualCamera/Source/VirtualCamera/Public/IVirtualCameraController.h
@@ -0,0 +1,90 @@
+// Copyright 1998-2019 Epic Games, Inc. All Rights Reserved.
+
+#pragma once
+
+#include "CoreMinimal.h"
+#include "UObject/Interface.h"
+#include "UObject/ScriptInterface.h"
+
+#include "IVirtualCameraController.generated.h"
+
+class IVirtualCameraOptions;
+class IVirtualCameraPresetContainer;
+class UCineCameraComponent;
+class ULevelSequencePlaybackController;
+
+UENUM(BlueprintType)
+enum class ETouchInputState : uint8
+{
+	/* Allows user to select an actor to always be in focus */
+	ActorFocusTargeting,
+
+	/* Allows user to select a point on the screen to auto-focus through */
+	AutoFocusTargeting,
+
+	/* Allows the touch input to be handled in the blueprint event. This should be the default */
+	BlueprintDefined,
+
+	/* Allows for the user to focus on target on touch without exiting manual focus */
+	ManualTouchFocus,
+
+	/* Touch support for scrubbing through a sequence */
+	Scrubbing,
+
+	/* Touch and hold for attach targeting */
+	TouchAndHold,
+};
+
+USTRUCT(BlueprintType)
+struct FTrackingOffset
+{
+	GENERATED_BODY()
+
+	UPROPERTY(EditAnywhere, BlueprintReadWrite, Category = "Camera Settings")
+	FVector Translation;
+
+	UPROPERTY(EditAnywhere, BlueprintReadWrite, Category = "Camera Settings")
+	FRotator Rotation;
+
+	FTransform AsTransform() const
+	{
+		return FTransform(Rotation, Translation);
+	}
+
+	FTrackingOffset()
+		: Translation(EForceInit::ForceInitToZero), Rotation(EForceInit::ForceInitToZero)
+	{
+
+	};
+};
+
+UINTERFACE(Blueprintable)
+class VIRTUALCAMERA_API UVirtualCameraController: public UInterface
+{
+	GENERATED_BODY()
+};
+
+class VIRTUALCAMERA_API IVirtualCameraController
+{
+	GENERATED_BODY()
+	
+public:
+
+	/** Returns the target camera that is used to create the streamed view. */
+	UFUNCTION(BlueprintNativeEvent, BlueprintCallable, Category = "VirtualCamera | Component")
+	UCineCameraComponent* GetStreamedCameraComponent() const;
+
+	/** Returns the recorded camera. */
+	UFUNCTION(BlueprintNativeEvent, BlueprintCallable, Category = "VirtualCamera | Component")
+	UCineCameraComponent* GetRecordingCameraComponent() const;
+
+	/** Returns the VirtualCamera's Sequence Controller. */
+	UFUNCTION(BlueprintNativeEvent, BlueprintCallable, Category = "VirtualCamera | Sequencer")
+	ULevelSequencePlaybackController* GetSequenceController() const;
+
+	UFUNCTION(BlueprintNativeEvent, BlueprintCallable, Category = "VirtualCamera | Presets")
+	TScriptInterface<IVirtualCameraPresetContainer> GetPresetContainer() const;
+
+	UFUNCTION(BlueprintNativeEvent, BlueprintCallable, Category = "VirtualCamera | Options")
+	TScriptInterface<IVirtualCameraOptions> GetOptions() const;
+};
diff --git a/Engine/Plugins/Experimental/VirtualCamera/Source/VirtualCamera/Public/IVirtualCameraOptions.h b/Engine/Plugins/Experimental/VirtualCamera/Source/VirtualCamera/Public/IVirtualCameraOptions.h
new file mode 100644
--- /dev/null
+++ b/Engine/Plugins/Experimental/VirtualCamera/Source/VirtualCamera/Public/IVirtualCameraOptions.h
@@ -0,0 +1,41 @@
+// Copyright 1998-2019 Epic Games, Inc. All Rights Reserved.
+
+#pragma once
+
+#include "UObject/Interface.h"
+
+#include "IVirtualCameraOptions.generated.h"
+
+UINTERFACE(Blueprintable)
+class VIRTUALCAMERA_API UVirtualCameraOptions : public UInterface
+{
+	GENERATED_BODY()
+};
+
+class VIRTUALCAMERA_API IVirtualCameraOptions
+{
+	GENERATED_BODY()
+	
+public:
+
+	/**
+	 * Sets unit of distance.
+	 * @param DesiredUnits - The new unit to use for distance measures like focus distance
+	 */
+	UFUNCTION(BlueprintNativeEvent, BlueprintCallable, Category = "VirtualCamera | Options")
+	void SetDesiredDistanceUnits(const EUnit DesiredUnits);
+
+	/**
+	 * Returns previously set unit of distance.
+	 * @return DesiredUnits - The unit that is used for distance measures like focus distance
+	 */
+	UFUNCTION(BlueprintNativeEvent, BlueprintCallable, Category = "VirtualCamera | Options")
+	EUnit GetDesiredDistanceUnits();
+
+	/**
+	 * Checks whether or not focus visualization can activate
+	 * @return the current state of touch event visualization
+	 */
+	UFUNCTION(BlueprintNativeEvent, BlueprintCallable, Category = "VirtualCamera | Options")
+	bool IsFocusVisualizationAllowed();
+};
diff --git a/Engine/Plugins/Experimental/VirtualCamera/Source/VirtualCamera/Public/IVirtualCameraPresetContainer.h b/Engine/Plugins/Experimental/VirtualCamera/Source/VirtualCamera/Public/IVirtualCameraPresetContainer.h
new file mode 100644
--- /dev/null
+++ b/Engine/Plugins/Experimental/VirtualCamera/Source/VirtualCamera/Public/IVirtualCameraPresetContainer.h
@@ -0,0 +1,55 @@
+// Copyright 1998-2019 Epic Games, Inc. All Rights Reserved.
+
+#pragma once
+
+#include "UObject/Interface.h"
+#include "VirtualCameraSaveGame.h"
+
+#include "IVirtualCameraPresetContainer.generated.h"
+
+UINTERFACE(Blueprintable)
+class VIRTUALCAMERA_API UVirtualCameraPresetContainer : public UInterface
+{
+	GENERATED_BODY()
+};
+
+class VIRTUALCAMERA_API IVirtualCameraPresetContainer
+{
+	GENERATED_BODY()
+	
+public:
+
+	/**
+	 * Saves a preset into the list of presets.
+	 * @param bSaveCameraSettings - Should this preset save camera settings
+	 * @param bSaveStabilization - Should this preset save stabilization settings
+	 * @param bSaveAxisLocking - Should this preset save axis locking settings
+	 * @param bSaveMotionScale - Should this preset save motion scaled settings
+	 * @return the name of the preset
+	 */
+	UFUNCTION(BlueprintNativeEvent, BlueprintCallable, Category = "VirtualCamera | Presets")
+	FString SavePreset(const bool bSaveCameraSettings, const bool bSaveStabilization, const bool bSaveAxisLocking, const bool bSaveMotionScale);
+
+	/**
+	 * Loads a preset using its name as a string key.
+	 * @param PresetName - The name of the preset to load
+	 * @return true if successful, false otherwise
+	 */
+	UFUNCTION(BlueprintNativeEvent, BlueprintCallable, Category = "VirtualCamera | Presets")
+	bool LoadPreset(const FString& PresetName);
+
+	/**
+	 * Deletes a preset using its name as the key.
+	 * @param PresetName - The name of the preset to delete
+	 * @return the number of values associated with the key
+	 */
+	UFUNCTION(BlueprintNativeEvent, BlueprintCallable, Category = "VirtualCamera | Presets")
+	int32 DeletePreset(const FString& PresetName);
+
+	/**
+	 * Returns a sorted TMap of the current presets.
+	 * @return a sorted TMap of settings presets
+	 */
+	UFUNCTION(BlueprintNativeEvent, BlueprintCallable, Category = "VirtualCamera | Presets")
+	TMap<FString, FVirtualCameraSettingsPreset>  GetSettingsPresets();
+};
diff --git a/Engine/Plugins/Experimental/VirtualCamera/Source/VirtualCamera/Public/LevelSequencePlaybackController.h b/Engine/Plugins/Experimental/VirtualCamera/Source/VirtualCamera/Public/LevelSequencePlaybackController.h
--- a/Engine/Plugins/Experimental/VirtualCamera/Source/VirtualCamera/Public/LevelSequencePlaybackController.h
+++ b/Engine/Plugins/Experimental/VirtualCamera/Source/VirtualCamera/Public/LevelSequencePlaybackController.h
@@ -44,7 +44,7 @@ public:
 	}
 };
 
-UCLASS()
+UCLASS(BlueprintType)
 class VIRTUALCAMERA_API ULevelSequencePlaybackController : public UObject
 {
 	GENERATED_UCLASS_BODY()
@@ -53,95 +53,107 @@ public:
 	/** Notify whether recording is enabled or disabled for the current sequence*/
 	FRecordEnabledStateChanged OnRecordEnabledStateChanged;
 
-	/**
-	 * Plays current level sequence from the current time.
-	 */
-	void ResumeLevelSequencePlay();
-
 	/**
 	 * Returns the names of each level sequence actor that is present in the level.
 	 * @param &OutLevelSequenceNames - UponReturn, array will contain all level sequence names.
 	 */
+	UFUNCTION(BlueprintCallable, Category = "VirtualCamera | Sequencer | Playback")
 	void GetLevelSequences(TArray<FLevelSequenceData>& OutLevelSequenceNames);
 
 	/**
 	 * @return the name of the currently selected sequence; returns empty string if no selected sequence
 	 */
+	UFUNCTION(BlueprintPure, Category = "VirtualCamera | Sequencer | Playback")
 	FString GetActiveLevelSequenceName() const;
 
 	/**
 	 * @return the currently selected LevelSequence
 	 */
+	UFUNCTION(BlueprintPure, Category = "VirtualCamera | Sequencer | Playback")
 	ULevelSequence* GetActiveLevelSequence() { return ActiveLevelSequence; }
 
 	/**
 	 * @return the FrameRate of the currently loaded sequence FrameRate
 	 */
+	UFUNCTION(BlueprintPure, Category = "VirtualCamera | Sequencer | Playback")
 	FFrameRate GetCurrentSequenceFrameRate() const;
 
 	/**
 	 * @return true if the active Sequencer is locked to camera cut
 	 */
+	UFUNCTION(BlueprintPure, Category = "VirtualCamera | Sequencer | Playback")
 	bool IsSequencerLockedToCameraCut() const;
 
 	/**
 	 * Sets the current Sequencer perspective to be locked to camera cut
 	 */
+	UFUNCTION(BlueprintCallable, Category = "VirtualCamera | Sequencer | Playback")
 	void SetSequencerLockedToCameraCut(bool bLockView);
 
 	/**
 	 * @return The FrameNumber of the sequence's start.
 	 */
+	UFUNCTION(BlueprintPure, Category = "VirtualCamera | Sequencer | Playback")
 	FFrameNumber GetCurrentSequencePlaybackStart() const;
 
 	/**
 	 * @return The FrameNumber of the sequence's end.
 	 */
+	UFUNCTION(BlueprintPure, Category = "VirtualCamera | Sequencer | Playback")
 	FFrameNumber GetCurrentSequencePlaybackEnd() const;
 
 	/**
 	 * @return the duration of the sequence in FrameNumber
 	 */
+	UFUNCTION(BlueprintPure, Category = "VirtualCamera | Sequencer | Playback")
 	FFrameNumber GetCurrentSequenceDuration() const;
 
 	/**
 	 * @return the current FrameTime of the sequence playback.
 	 */
+	UFUNCTION(BlueprintPure, Category = "VirtualCamera | Sequencer | Playback")
 	FFrameTime GetCurrentSequencePlaybackPosition() const;
 
 	/**
 	 * @return The current Timecode of the sequence playback.
 	 */
+	UFUNCTION(BlueprintPure, Category = "VirtualCamera | Sequencer | Playback")
 	FTimecode GetCurrentSequencePlaybackTimecode() const;
 
 	/**
 	 * Moves the current sequence to a desired playback position
 	 */
+	UFUNCTION(BlueprintCallable, Category = "VirtualCamera | Sequencer | Playback")
 	void JumpToPlaybackPosition(const FFrameNumber& InFrameNumber);
 
 	/**
 	 * @return true if a valid LevelSequence is being played.
 	 */
+	UFUNCTION(BlueprintPure, Category = "VirtualCamera | Sequencer | Playback")
 	bool IsSequencePlaybackActive() const;
 
 	/**
 	 * Pause the currently active sequence 
 	 */
+	UFUNCTION(BlueprintCallable, Category = "VirtualCamera | Sequencer | Playback")
 	void PauseLevelSequence();
 
 	/**
 	 * Starts playing the currently active sequence
 	 */
+	UFUNCTION(BlueprintCallable, Category = "VirtualCamera | Sequencer | Playback")
 	void PlayLevelSequence();
 
 	/**
 	 * Starts playing the currently active sequence in reverse
 	 */
+	UFUNCTION(BlueprintCallable, Category = "VirtualCamera | Sequencer | Playback")
 	void PlayLevelSequenceReverse();
 
 	/**
 	 * Stops playing the currently active sequence
 	 */
+	UFUNCTION(BlueprintCallable, Category = "VirtualCamera | Sequencer | Playback")
 	void StopLevelSequencePlay();
 
 	/**
@@ -149,18 +161,15 @@ public:
 	 * @param InNewLevelSequence - The LevelSequence to select
 	 * @return true if a valid LevelSequence was passed and sequencer was successfully found.
 	 */
+	UFUNCTION(BlueprintCallable, Category = "VirtualCamera | Sequencer | Playback")
 	bool SetActiveLevelSequence(ULevelSequence* InNewLevelSequence);
 
 	/**
 	 * Clears the current level sequence player, needed when recording clean takes of something.
 	 */
+	UFUNCTION(BlueprintCallable, Category = "VirtualCamera | Sequencer | Playback")
 	void ClearActiveLevelSequence();
 
-	/**
-	 * Plays current level sequence from beginning.
-	 */
-	void PlayFromBeginning();
-
 protected:
 
 	/** The sequence to play back */
diff --git a/Engine/Plugins/Experimental/VirtualCamera/Source/VirtualCamera/Public/VirtualCameraPawnBase.h b/Engine/Plugins/Experimental/VirtualCamera/Source/VirtualCamera/Public/VirtualCameraPawnBase.h
--- a/Engine/Plugins/Experimental/VirtualCamera/Source/VirtualCamera/Public/VirtualCameraPawnBase.h
+++ b/Engine/Plugins/Experimental/VirtualCamera/Source/VirtualCamera/Public/VirtualCameraPawnBase.h
@@ -3,6 +3,8 @@
 #pragma once
 
 #include "GameFramework/Pawn.h"
+#include "IVirtualCameraOptions.h"
+#include "IVirtualCameraPresetContainer.h"
 #include "VirtualCameraSaveGame.h"
 #include "VirtualCameraCineCameraComponent.h"
 #include "VirtualCameraMovementComponent.h"
@@ -15,7 +17,7 @@ class USceneComponent;
  * A class to handle aspects of virtual Camera related to general settings, and communicating with components.
  */
 UCLASS(Abstract, Blueprintable, BlueprintType, HideCategories = ("Pawn", "Camera", "Rendering", "Replication", "Input", "Actor", "HLOD"))
-class VIRTUALCAMERA_API AVirtualCameraPawnBase : public APawn
+class VIRTUALCAMERA_API AVirtualCameraPawnBase : public APawn, public IVirtualCameraPresetContainer, public IVirtualCameraOptions
 {
 	GENERATED_UCLASS_BODY()
 
@@ -323,4 +325,17 @@ protected:
 	 * @return The input number padded to be at least MinNumberOfCharacters in length
 	 */
 	FString LeftPadWithZeros(int32 InNumber, int32 MinNumberOfCharacters) const;
+
+	//~ Begin IVirtualCameraPresetContainer Interface
+	virtual FString SavePreset_Implementation(const bool bSaveCameraSettings, const bool bSaveStabilization, const bool bSaveAxisLocking, const bool bSaveMotionScale) override;
+	virtual bool LoadPreset_Implementation(const FString& PresetName) override;
+	virtual int32 DeletePreset_Implementation(const FString& PresetName) override;
+	virtual TMap<FString, FVirtualCameraSettingsPreset> GetSettingsPresets_Implementation() override;
+	//~ End IVirtualCameraPresetContainer Interface
+
+	//~ Begin IVirtualCameraOptions Interface
+	virtual void SetDesiredDistanceUnits_Implementation(const EUnit DesiredUnits) override;
+	virtual EUnit GetDesiredDistanceUnits_Implementation() override;
+	virtual bool IsFocusVisualizationAllowed_Implementation() override;
+	//~ End IVirtualCameraOptions Interface
 };
diff --git a/Engine/Plugins/Experimental/VirtualCamera/Source/VirtualCamera/Public/VirtualCameraPlayerControllerBase.h b/Engine/Plugins/Experimental/VirtualCamera/Source/VirtualCamera/Public/VirtualCameraPlayerControllerBase.h
--- a/Engine/Plugins/Experimental/VirtualCamera/Source/VirtualCamera/Public/VirtualCameraPlayerControllerBase.h
+++ b/Engine/Plugins/Experimental/VirtualCamera/Source/VirtualCamera/Public/VirtualCameraPlayerControllerBase.h
@@ -7,6 +7,7 @@
 #include "CineCameraActor.h"
 #include "ILiveLinkClient.h"
 #include "InputCore.h"
+#include "IVirtualCameraController.h"
 #include "LevelSequencePlaybackController.h"
 #include "Math/UnitConversion.h"
 #include "Misc/FrameNumber.h"
@@ -30,54 +31,8 @@ enum class ETrackerInputSource : uint8
 	LiveLink
 };
 
-UENUM(BlueprintType)
-enum class ETouchInputState : uint8
-{
-	/* Allows user to select an actor to always be in focus */
-	ActorFocusTargeting,
-
-	/* Allows user to select a point on the screen to auto-focus through */
-	AutoFocusTargeting,
-
-	/* Allows the touch input to be handled in the blueprint event. This should be the default */
-	BlueprintDefined,
-
-	/* Allows for the user to focus on target on touch without exiting manual focus */
-	ManualTouchFocus,
-
-	/* Touch support for scrubbing through a sequence */
-	Scrubbing,
-
-	/* Touch and hold for attach targeting */
-	TouchAndHold,
-};
-
-USTRUCT(BlueprintType)
-struct FTrackingOffset
-{
-public:
-	GENERATED_BODY()
-	
-	UPROPERTY(EditAnywhere, BlueprintReadWrite, Category = "Camera Settings")
-	FVector Translation;
-	
-	UPROPERTY(EditAnywhere, BlueprintReadWrite, Category = "Camera Settings")
-	FRotator Rotation;
-	
-	FTransform AsTransform() const
-	{
-		return FTransform(Rotation, Translation);
-	}
-	
-	FTrackingOffset()
-		: Translation(EForceInit::ForceInitToZero), Rotation(EForceInit::ForceInitToZero)
-	{
-		
-	};
-};
-
 UCLASS(Abstract)
-class VIRTUALCAMERA_API AVirtualCameraPlayerControllerBase : public APlayerController
+class VIRTUALCAMERA_API AVirtualCameraPlayerControllerBase : public APlayerController, public IVirtualCameraController
 {
 	GENERATED_UCLASS_BODY()
 
@@ -492,8 +447,8 @@ public:
 	FFrameRate GetCurrentSequenceFrameRate();
 
 	/**
-	 * Set the matte aspect ratio to a new value.
-	 * @return DesiredUnits - The new unit to use for distance measures like focus distance
+	 * Returns previously set unit of distance.
+	 * @return DesiredUnits - The unit that is used for distance measures like focus distance
 	 */
 	UFUNCTION(BlueprintPure, Category = "Virtual Camera|Settings")
 	EUnit GetDesiredDistanceUnits();
@@ -777,7 +732,7 @@ public:
 	void SetCurrentFocusDistance(const float NewFocusDistance);
 
 	/**
-	 * Set the matte aspect ratio to a new value.
+	 * Sets unit of distance.
 	 * @param DesiredUnits - The new unit to use for distance measures like focus distance
 	 */
 	UFUNCTION(BlueprintCallable, Category = "Virtual Camera|Settings")
@@ -948,4 +903,13 @@ public:
 	 */
 	UFUNCTION(BlueprintCallable, Category = "Virtual Camera")
 	void ActivateGameViewport();
+
+private:
+	//~ Begin IVirtualCameraController Interface
+	virtual UCineCameraComponent* GetStreamedCameraComponent_Implementation() const override;
+	virtual UCineCameraComponent* GetRecordingCameraComponent_Implementation() const override;
+	virtual ULevelSequencePlaybackController* GetSequenceController_Implementation() const override;
+	virtual TScriptInterface<IVirtualCameraPresetContainer> GetPresetContainer_Implementation() const override;
+	virtual TScriptInterface<IVirtualCameraOptions> GetOptions_Implementation() const override;
+	//~ End  IVirtualCameraController Interface
 };
diff --git a/Engine/Plugins/Experimental/VirtualCamera/Source/VirtualCamera/Public/VirtualCameraSubsystem.h b/Engine/Plugins/Experimental/VirtualCamera/Source/VirtualCamera/Public/VirtualCameraSubsystem.h
new file mode 100644
--- /dev/null
+++ b/Engine/Plugins/Experimental/VirtualCamera/Source/VirtualCamera/Public/VirtualCameraSubsystem.h
@@ -0,0 +1,27 @@
+// Copyright 1998-2019 Epic Games, Inc. All Rights Reserved.
+
+#pragma once
+
+#include "IVirtualCameraController.h"
+#include "Subsystems/EngineSubsystem.h"
+
+#include "VirtualCameraSubsystem.generated.h"
+
+UCLASS()
+class VIRTUALCAMERA_API UVirtualCameraSubsystem : public UEngineSubsystem
+{
+	GENERATED_BODY()
+	
+public:
+
+	UFUNCTION(BlueprintCallable, Category = "Virtual Camera")
+	TScriptInterface<IVirtualCameraController> GetVirtualCameraController() const { return ActiveCameraController; }
+	
+	UFUNCTION(BlueprintCallable, Category = "Virtual Camera")
+	void SetVirtualCameraController(TScriptInterface<IVirtualCameraController> VirtualCamera);
+	
+private:
+
+	UPROPERTY(Transient)
+	TScriptInterface<IVirtualCameraController> ActiveCameraController;
+};
