Added missing file for VirtualCamera Plugin

diff --git a/Engine/Plugins/Experimental/VirtualCamera/Source/VirtualCamera/Public/VirtualCameraVersion.h b/Engine/Plugins/Experimental/VirtualCamera/Source/VirtualCamera/Public/VirtualCameraVersion.h
new file mode 100644
--- /dev/null
+++ b/Engine/Plugins/Experimental/VirtualCamera/Source/VirtualCamera/Public/VirtualCameraVersion.h
@@ -0,0 +1,5 @@
+// Copyright 1998-2019 Epic Games, Inc. All Rights Reserved.
+
+#pragma once
+
+#define VIRTUAL_CAMERA_VERSION_STRING TEXT("1.0.0")
