Cine Camera: Change depth of field method None to Do Not Override to describe what it is actually doing - not overriding post process settings. Added a new method Disable which disables depth of field entirely by overriding post process settings.

diff --git a/Engine/Plugins/Enterprise/DatasmithImporter/Source/DatasmithImporter/Private/DatasmithCameraImporter.cpp b/Engine/Plugins/Enterprise/DatasmithImporter/Source/DatasmithImporter/Private/DatasmithCameraImporter.cpp
--- a/Engine/Plugins/Enterprise/DatasmithImporter/Source/DatasmithImporter/Private/DatasmithCameraImporter.cpp
+++ b/Engine/Plugins/Enterprise/DatasmithImporter/Source/DatasmithImporter/Private/DatasmithCameraImporter.cpp
@@ -110,7 +110,7 @@ void FDatasmithCameraImporter::SetupCineCameraComponent( UCineCameraComponent* C
 	CameraTemplate->CurrentFocalLength = CameraElement->GetFocalLength();
 	CameraTemplate->CurrentAperture = CameraElement->GetFStop();
 
-	CameraTemplate->FocusSettings.FocusMethod = CameraElement->GetEnableDepthOfField() ? ECameraFocusMethod::Manual : ECameraFocusMethod::None;
+	CameraTemplate->FocusSettings.FocusMethod = CameraElement->GetEnableDepthOfField() ? ECameraFocusMethod::Manual : ECameraFocusMethod::DoNotOverride;
 	CameraTemplate->FocusSettings.ManualFocusDistance = CameraElement->GetFocusDistance();
 
 	CameraTemplate->PostProcessSettings = FDatasmithPostProcessImporter::CopyDatasmithPostProcessToUEPostProcess( CameraElement->GetPostProcess() );
diff --git a/Engine/Plugins/Experimental/VirtualCamera/Source/VirtualCamera/Private/VirtualCameraCineCameraComponent.cpp b/Engine/Plugins/Experimental/VirtualCamera/Source/VirtualCamera/Private/VirtualCameraCineCameraComponent.cpp
--- a/Engine/Plugins/Experimental/VirtualCamera/Source/VirtualCamera/Private/VirtualCameraCineCameraComponent.cpp
+++ b/Engine/Plugins/Experimental/VirtualCamera/Source/VirtualCamera/Private/VirtualCameraCineCameraComponent.cpp
@@ -180,9 +180,9 @@ void UVirtualCameraCineCameraComponent::SetFocusMethod(const EVirtualCameraFocus
 
 	switch (NewFocusMethod)
 	{
-	case EVirtualCameraFocusMethod::None:
+	case EVirtualCameraFocusMethod::DoNotOverride:
 		bAutoFocusEnabled = false;
-		FocusSettings.FocusMethod = ECameraFocusMethod::None;
+		FocusSettings.FocusMethod = ECameraFocusMethod::DoNotOverride;
 		break;
 	case EVirtualCameraFocusMethod::Auto:
 		bAutoFocusEnabled = true;
@@ -196,6 +196,10 @@ void UVirtualCameraCineCameraComponent::SetFocusMethod(const EVirtualCameraFocus
 		bAutoFocusEnabled = false;
 		FocusSettings.FocusMethod = ECameraFocusMethod::Tracking;
 		break;
+	case EVirtualCameraFocusMethod::Disable:
+		bAutoFocusEnabled = false;
+		FocusSettings.FocusMethod = ECameraFocusMethod::Disable;
+		break;
 	default:  // Should never be reached, but just in case new focus methods are added
 		UE_LOG(LogActor, Warning, TEXT("Specified focus method is not currently supported in Virtual Camera!"))
 		break;
@@ -229,7 +233,7 @@ void UVirtualCameraCineCameraComponent::SetFocusChangeSmoothness(const float New
 
 void UVirtualCameraCineCameraComponent::SetFocusVisualization(bool bShowFocusVisualization)
 {
-	if (FocusSettings.FocusMethod == ECameraFocusMethod::None)
+	if (FocusSettings.FocusMethod == ECameraFocusMethod::Disable)
 	{
 		UE_LOG(LogActor, Warning, TEXT("Camera focus mode is currently set to none, cannot display focus plane!"))
 		return;
@@ -240,7 +244,7 @@ void UVirtualCameraCineCameraComponent::SetFocusVisualization(bool bShowFocusVis
 
 bool UVirtualCameraCineCameraComponent::IsFocusVisualizationActivated() const
 {
-	return FocusSettings.FocusMethod != ECameraFocusMethod::None && FocusSettings.bDrawDebugFocusPlane;
+	return FocusSettings.FocusMethod != ECameraFocusMethod::Disable && FocusSettings.bDrawDebugFocusPlane;
 }
 
 int32 UVirtualCameraCineCameraComponent::FindClosestPresetIndex(const TArray<float>& ArrayToSearch, const float SearchValue) const
diff --git a/Engine/Plugins/Experimental/VirtualCamera/Source/VirtualCamera/Public/VirtualCameraCineCameraComponent.h b/Engine/Plugins/Experimental/VirtualCamera/Source/VirtualCamera/Public/VirtualCameraCineCameraComponent.h
--- a/Engine/Plugins/Experimental/VirtualCamera/Source/VirtualCamera/Public/VirtualCameraCineCameraComponent.h
+++ b/Engine/Plugins/Experimental/VirtualCamera/Source/VirtualCamera/Public/VirtualCameraCineCameraComponent.h
@@ -8,8 +8,8 @@
 UENUM()
 enum class EVirtualCameraFocusMethod : uint8
 {
-	/* Depth of Field disabled entirely */
-	None,
+	/** Don't override, ie. allow post process volume settings to persist. */
+	DoNotOverride,
 
 	/* User controls focus distance directly */
 	Manual,
@@ -19,6 +19,9 @@ enum class EVirtualCameraFocusMethod : uint8
 
 	/* Focus distance automatically changes to focus on actors in a specific screen location */
 	Auto,
+
+	/** Disable depth of field entirely. */
+	Disable,
 };
 
 UCLASS(Blueprintable, BlueprintType)
diff --git a/Engine/Plugins/Importers/USDImporter/Source/USDUtilities/Private/USDPrimConversion.cpp b/Engine/Plugins/Importers/USDImporter/Source/USDUtilities/Private/USDPrimConversion.cpp
--- a/Engine/Plugins/Importers/USDImporter/Source/USDUtilities/Private/USDPrimConversion.cpp
+++ b/Engine/Plugins/Importers/USDImporter/Source/USDUtilities/Private/USDPrimConversion.cpp
@@ -94,7 +94,7 @@ bool UsdToUnreal::ConvertGeomCamera( const pxr::UsdStageRefPtr& Stage, const pxr
 
 	if ( FMath::IsNearlyZero( CameraComponent.FocusSettings.ManualFocusDistance ) )
 	{
-		CameraComponent.FocusSettings.FocusMethod = ECameraFocusMethod::None;
+		CameraComponent.FocusSettings.FocusMethod = ECameraFocusMethod::DoNotOverride;
 	}
 
 	CameraComponent.CurrentAperture = UsdUtils::GetUsdValue< float >( GeomCamera.GetFStopAttr(), EvalTime );
diff --git a/Engine/Source/Editor/DetailCustomizations/Private/CameraFocusSettingsCustomization.cpp b/Engine/Source/Editor/DetailCustomizations/Private/CameraFocusSettingsCustomization.cpp
--- a/Engine/Source/Editor/DetailCustomizations/Private/CameraFocusSettingsCustomization.cpp
+++ b/Engine/Source/Editor/DetailCustomizations/Private/CameraFocusSettingsCustomization.cpp
@@ -181,9 +181,9 @@ EVisibility FCameraFocusSettingsCustomization::IsGeneralSettingGroupVisible() co
 	uint8 FocusMethodNumber;
 	FocusMethodHandle->GetValue(FocusMethodNumber);
 	ECameraFocusMethod const FocusMethod = static_cast<ECameraFocusMethod>(FocusMethodNumber);
-	if (FocusMethod != ECameraFocusMethod::None)
+	if (FocusMethod != ECameraFocusMethod::Disable && FocusMethod != ECameraFocusMethod::DoNotOverride)
 	{
-		// if focus method is set to none, all non-none setting groups are collapsed
+		// if focus method is set to disable, all non-none setting groups are collapsed
 		return EVisibility::Visible;
 	}
 
diff --git a/Engine/Source/Editor/UnrealEd/Private/Fbx/FbxSceneImportFactory.cpp b/Engine/Source/Editor/UnrealEd/Private/Fbx/FbxSceneImportFactory.cpp
--- a/Engine/Source/Editor/UnrealEd/Private/Fbx/FbxSceneImportFactory.cpp
+++ b/Engine/Source/Editor/UnrealEd/Private/Fbx/FbxSceneImportFactory.cpp
@@ -1293,7 +1293,7 @@ USceneComponent *CreateCameraComponent(AActor *ParentActor, TSharedPtr<FFbxCamer
 	CameraComponent->Filmback.SensorHeight = FUnitConversion::Convert(CameraInfo->ApertureHeight, EUnit::Inches, EUnit::Millimeters);
 	CameraComponent->LensSettings.MaxFocalLength = CameraInfo->FocalLength;
 	CameraComponent->LensSettings.MinFocalLength = CameraInfo->FocalLength;
-	CameraComponent->FocusSettings.FocusMethod = ECameraFocusMethod::None;
+	CameraComponent->FocusSettings.FocusMethod = ECameraFocusMethod::DoNotOverride;
 
 	return CameraComponent;
 }
diff --git a/Engine/Source/Runtime/CinematicCamera/Private/CineCameraComponent.cpp b/Engine/Source/Runtime/CinematicCamera/Private/CineCameraComponent.cpp
--- a/Engine/Source/Runtime/CinematicCamera/Private/CineCameraComponent.cpp
+++ b/Engine/Source/Runtime/CinematicCamera/Private/CineCameraComponent.cpp
@@ -130,6 +130,11 @@ void UCineCameraComponent::PostLoad()
 		}
 	}
 
+	if (FocusSettings.FocusMethod >= ECameraFocusMethod::MAX )
+	{
+		FocusSettings.FocusMethod = ECameraFocusMethod::DoNotOverride;
+	}
+
 	RecalcDerivedData();
 	bResetInterpolation = true;
 	Super::PostLoad();
@@ -476,7 +481,7 @@ void UCineCameraComponent::UpdateDebugFocusPlane()
 
 void UCineCameraComponent::UpdateCameraLens(float DeltaTime, FMinimalViewInfo& DesiredView)
 {
-	if (FocusSettings.FocusMethod == ECameraFocusMethod::None)
+	if (FocusSettings.FocusMethod == ECameraFocusMethod::DoNotOverride)
 	{
 		DesiredView.PostProcessSettings.bOverride_DepthOfFieldFstop = false;
 		DesiredView.PostProcessSettings.bOverride_DepthOfFieldMinFstop = false;
@@ -484,6 +489,12 @@ void UCineCameraComponent::UpdateCameraLens(float DeltaTime, FMinimalViewInfo& D
 		DesiredView.PostProcessSettings.bOverride_DepthOfFieldFocalDistance = false;
 		DesiredView.PostProcessSettings.bOverride_DepthOfFieldSensorWidth = false;
 	}
+	else if (FocusSettings.FocusMethod == ECameraFocusMethod::Disable)
+	{
+		// There might be a post process volume that is enabled with depth of field settings, override it and disable depth of field by setting the distance to 0
+		DesiredView.PostProcessSettings.bOverride_DepthOfFieldFocalDistance = true;
+		DesiredView.PostProcessSettings.DepthOfFieldFocalDistance = 0.f;
+	}
 	else
 	{
 		// Update focus/DoF
diff --git a/Engine/Source/Runtime/CinematicCamera/Public/CineCameraComponent.h b/Engine/Source/Runtime/CinematicCamera/Public/CineCameraComponent.h
--- a/Engine/Source/Runtime/CinematicCamera/Public/CineCameraComponent.h
+++ b/Engine/Source/Runtime/CinematicCamera/Public/CineCameraComponent.h
@@ -120,14 +120,19 @@ struct FNamedLensPreset
 UENUM()
 enum class ECameraFocusMethod : uint8
 {
-	/** Disables DoF entirely. */
-	None,
+	/** Don't override, ie. allow post process volume settings to persist. */
+	DoNotOverride,
 
 	/** Allows for specifying or animating exact focus distances. */
 	Manual,
 
 	/** Locks focus to specific object. */
 	Tracking,
+
+	/** Disable depth of field entirely. */
+	Disable,
+
+	MAX UMETA(Hidden)
 };
 
 /** Settings to control tracking-focus mode. */
