Fix a couple of deprecation warnings.

diff --git a/Engine/Plugins/Experimental/VirtualCamera/Source/VirtualCamera/Private/VCamBlueprintFunctionLibrary.cpp b/Engine/Plugins/Experimental/VirtualCamera/Source/VirtualCamera/Private/VCamBlueprintFunctionLibrary.cpp
--- a/Engine/Plugins/Experimental/VirtualCamera/Source/VirtualCamera/Private/VCamBlueprintFunctionLibrary.cpp
+++ b/Engine/Plugins/Experimental/VirtualCamera/Source/VirtualCamera/Private/VCamBlueprintFunctionLibrary.cpp
@@ -15,6 +15,7 @@
 #include "EditorAssetLibrary.h"
 #include "EditorLevelLibrary.h"
 #endif
+#include <Editor\LevelEditor\Public\LevelEditorSubsystem.h>
 
 bool UVCamBlueprintFunctionLibrary::IsGameRunning()
 {
@@ -246,7 +247,11 @@ void UVCamBlueprintFunctionLibrary::PilotActor(AActor* SelectedActor)
 #if WITH_EDITOR
 	if (SelectedActor)
 	{
-		UEditorLevelLibrary::PilotLevelActor(SelectedActor);
+		ULevelEditorSubsystem* LevelEditorSubsystem = GEditor->GetEditorSubsystem<ULevelEditorSubsystem>();
+		if (LevelEditorSubsystem)
+		{
+			return LevelEditorSubsystem->PilotLevelActor(SelectedActor);
+		}
 	}
 #endif
 }
@@ -311,6 +316,11 @@ void UVCamBlueprintFunctionLibrary::EditorSetGameView(bool bIsToggled)
 	{
 		return;
 	}
-	UEditorLevelLibrary::EditorSetGameView(bIsToggled);
+
+	ULevelEditorSubsystem* LevelEditorSubsystem = GEditor->GetEditorSubsystem<ULevelEditorSubsystem>();
+	if (LevelEditorSubsystem)
+	{
+		return LevelEditorSubsystem->EditorSetGameView(bIsToggled);
+	}
 #endif
 }
