Missing file.

diff --git a/Engine/Plugins/Experimental/VirtualCamera/Source/VirtualCamera/Private/VCamBlueprintFunctionLibrary.cpp b/Engine/Plugins/Experimental/VirtualCamera/Source/VirtualCamera/Private/VCamBlueprintFunctionLibrary.cpp
--- a/Engine/Plugins/Experimental/VirtualCamera/Source/VirtualCamera/Private/VCamBlueprintFunctionLibrary.cpp
+++ b/Engine/Plugins/Experimental/VirtualCamera/Source/VirtualCamera/Private/VCamBlueprintFunctionLibrary.cpp
@@ -14,8 +14,8 @@
 #include "Editor.h"
 #include "EditorAssetLibrary.h"
 #include "EditorLevelLibrary.h"
+#include "LevelEditor\Public\LevelEditorSubsystem.h"
 #endif
-#include <Editor\LevelEditor\Public\LevelEditorSubsystem.h>
 
 bool UVCamBlueprintFunctionLibrary::IsGameRunning()
 {
