Changing the GetModuleChecked to LoadModuleChecked.

diff --git a/Engine/Plugins/Experimental/VirtualCamera/Source/VirtualCamera/Private/VirtualCamera.cpp b/Engine/Plugins/Experimental/VirtualCamera/Source/VirtualCamera/Private/VirtualCamera.cpp
--- a/Engine/Plugins/Experimental/VirtualCamera/Source/VirtualCamera/Private/VirtualCamera.cpp
+++ b/Engine/Plugins/Experimental/VirtualCamera/Source/VirtualCamera/Private/VirtualCamera.cpp
@@ -33,7 +33,7 @@ private:
 		ConcertManager = MakeUnique<FConcertVirtualCameraManager>();
 
 		// Loads widgets (ex. RadialSlider) that are potentially referenced by assets
-		FModuleManager::Get().GetModuleChecked<FAdvancedWidgetsModule>("AdvancedWidgets");
+		FModuleManager::Get().LoadModuleChecked<FAdvancedWidgetsModule>("AdvancedWidgets");
 	}
 
 	virtual void ShutdownModule() override
