Prototype Wrapped Object Pointers (ObjectHandle/ObjectPtr)

Call-site changes affecting compilation of QAGame for non-scalar wrapped object pointer upgrades in Engine + Shootergame

diff --git a/Engine/Plugins/Enterprise/DataprepEditor/Source/DataprepEditor/Private/DataprepSnapshot.cpp b/Engine/Plugins/Enterprise/DataprepEditor/Source/DataprepEditor/Private/DataprepSnapshot.cpp
--- a/Engine/Plugins/Enterprise/DataprepEditor/Source/DataprepEditor/Private/DataprepSnapshot.cpp
+++ b/Engine/Plugins/Enterprise/DataprepEditor/Source/DataprepEditor/Private/DataprepSnapshot.cpp
@@ -854,7 +854,7 @@ void FDataprepEditor::RestoreFromSnapshot(bool bUpdateViewport)
 		GWorld = PreviewWorld;
 
 		// Cache and disable recording of transaction
-		TGuardValue<UTransactor*> NormalTransactor( GEditor->Trans, nullptr );
+		TGuardValue<decltype(GEditor->Trans)> NormalTransactor( GEditor->Trans, nullptr );
 
 		// Cache and disable warnings from LogExec because ULevelFactory::FactoryCreateText is pretty verbose on harmless warnings
 		ELogVerbosity::Type PrevLogExecVerbosity = LogExec.GetVerbosity();
diff --git a/Engine/Plugins/Experimental/VirtualCamera/Source/VirtualCamera/Private/VirtualCameraMovementComponent.cpp b/Engine/Plugins/Experimental/VirtualCamera/Source/VirtualCamera/Private/VirtualCameraMovementComponent.cpp
--- a/Engine/Plugins/Experimental/VirtualCamera/Source/VirtualCamera/Private/VirtualCameraMovementComponent.cpp
+++ b/Engine/Plugins/Experimental/VirtualCamera/Source/VirtualCamera/Private/VirtualCameraMovementComponent.cpp
@@ -89,8 +89,8 @@ void UVirtualCameraMovementComponent::ProcessMovementInput(const FVector& Tracke
 	// Apply the location change for the RootUpdatedComponent
 	if (RootUpdatedComponent)
 	{
-		TGuardValue<USceneComponent*> UpdateComponentGuard(UpdatedComponent, RootUpdatedComponent);
-		TGuardValue<UPrimitiveComponent*> UpdatePrimitiveGuard(UpdatedPrimitive, RootUpdatedPrimitive);
+		TGuardValue<decltype(UpdatedComponent)> UpdateComponentGuard(UpdatedComponent, RootUpdatedComponent);
+		TGuardValue<decltype(UpdatedPrimitive)> UpdatePrimitiveGuard(UpdatedPrimitive, RootUpdatedPrimitive);
 		FHitResult OutResult = FHitResult();
 		SafeMoveUpdatedComponent(FromControllerTargetLocation, FQuat::Identity, false, OutResult);
 		FromControllerTargetLocation = FVector::ZeroVector;
diff --git a/Engine/Plugins/Runtime/Steam/SteamSockets/Source/SteamSockets/Private/SteamSocketsNetDriver.cpp b/Engine/Plugins/Runtime/Steam/SteamSockets/Source/SteamSockets/Private/SteamSocketsNetDriver.cpp
--- a/Engine/Plugins/Runtime/Steam/SteamSockets/Source/SteamSockets/Private/SteamSocketsNetDriver.cpp
+++ b/Engine/Plugins/Runtime/Steam/SteamSockets/Source/SteamSockets/Private/SteamSocketsNetDriver.cpp
@@ -432,7 +432,7 @@ void USteamSocketsNetDriver::ResetSocketInfo(const FSteamSocket* RemovedSocket)
 
 UNetConnection* USteamSocketsNetDriver::FindClientConnectionForHandle(SteamSocketHandles SocketHandle)
 {
-	for (UNetConnection*& ClientConnection : ClientConnections)
+	for (UE_TRANSITIONAL_OBJECT_PTR(UNetConnection)& ClientConnection : ClientConnections)
 	{
 		USteamSocketsNetConnection* SteamConnection = static_cast<USteamSocketsNetConnection*>(ClientConnection);
 		if (SteamConnection && SteamConnection->GetRawSocket() != nullptr)
