Fix null access calling game-only functions on a non-game preview world
Issue was pointed out by [REDACTED_PERSON] on github as my CL 17208151 triggered a crash for them

diff --git a/Engine/Plugins/Experimental/VirtualCamera/Source/VCamCore/Private/VCamComponent.cpp b/Engine/Plugins/Experimental/VirtualCamera/Source/VCamCore/Private/VCamComponent.cpp
--- a/Engine/Plugins/Experimental/VirtualCamera/Source/VCamCore/Private/VCamComponent.cpp
+++ b/Engine/Plugins/Experimental/VirtualCamera/Source/VCamCore/Private/VCamComponent.cpp
@@ -860,7 +860,7 @@ void UVCamComponent::UpdateActorLock()
 #endif
 		{
 			UWorld* ActorWorld = Context.World();
-			if (ActorWorld)
+			if (ActorWorld && ActorWorld->GetGameInstance())
 			{
 				APlayerController* PlayerController = ActorWorld->GetGameInstance()->GetFirstLocalPlayerController(ActorWorld);
 				if (PlayerController)
