Remove Legacy code that was causing problems with PIE

diff --git a/Engine/Plugins/Experimental/VirtualCamera/Source/VirtualCamera/Private/VirtualCameraActor.cpp b/Engine/Plugins/Experimental/VirtualCamera/Source/VirtualCamera/Private/VirtualCameraActor.cpp
--- a/Engine/Plugins/Experimental/VirtualCamera/Source/VirtualCamera/Private/VirtualCameraActor.cpp
+++ b/Engine/Plugins/Experimental/VirtualCamera/Source/VirtualCamera/Private/VirtualCameraActor.cpp
@@ -216,11 +216,6 @@ void AVirtualCameraActor::Destroyed()
 	{
 		CameraScreenWidget->Hide();
 	}
-
-	if (RemoteSessionHost && RemoteSessionHost->IsConnected())
-	{
-		RemoteSessionHost->Close(TEXT("Camera Destroyed"));
-	}
 }
 
 #if WITH_EDITOR
@@ -309,11 +304,6 @@ void AVirtualCameraActor::Tick(float DeltaSeconds)
 {
 	Super::Tick(DeltaSeconds);
 
-	if (RemoteSessionHost)
-	{
-		RemoteSessionHost->Tick(DeltaSeconds);
-	}
-
 	if (CameraScreenWidget && CameraUMGClass)
 	{
 		CameraScreenWidget->Tick(DeltaSeconds);
@@ -429,21 +419,6 @@ bool AVirtualCameraActor::StartStreaming()
 		CameraScreenWidget->Display(ActorWorld);
 	}
 
-	if (IRemoteSessionModule* RemoteSession = FModuleManager::LoadModulePtr<IRemoteSessionModule>("RemoteSession"))
-	{
-		TArray<FRemoteSessionChannelInfo> SupportedChannels;
-		SupportedChannels.Emplace(FRemoteSessionInputChannel::StaticType(), ERemoteSessionChannelMode::Read);
-		SupportedChannels.Emplace(FRemoteSessionImageChannel::StaticType(), ERemoteSessionChannelMode::Write);
-
-		RemoteSessionHost = RemoteSession->CreateHost(MoveTemp(SupportedChannels), RemoteSessionPort);
-
-		if (RemoteSessionHost)
-		{
-			RemoteSessionHost->RegisterChannelChangeDelegate(FOnRemoteSessionChannelChange::FDelegate::CreateUObject(this, &AVirtualCameraActor::OnRemoteSessionChannelChange));
-			RemoteSessionHost->Tick(0.0f);
-		}
-	}
-
 	SetActorTickEnabled(true);
 
 	bIsStreaming = true;
@@ -452,19 +427,6 @@ bool AVirtualCameraActor::StartStreaming()
 
 bool AVirtualCameraActor::StopStreaming()
 {
-	if (!RemoteSessionHost)
-	{
-		return false;
-	}
-
-	TSharedPtr<FRemoteSessionInputChannel> InputChannel = RemoteSessionHost->GetChannel<FRemoteSessionInputChannel>();
-	if (InputChannel)
-	{
-		InputChannel->GetOnRouteTouchDownToWidgetFailedDelegate()->RemoveAll(this);
-	}
-
-	RemoteSessionHost.Reset();
-
 	CameraScreenWidget->Hide();
 	if (MediaCapture)
 	{
@@ -661,54 +623,6 @@ bool AVirtualCameraActor::IsFocusVisualizationAllowed_Implementation()
 	return bAllowFocusVisualization;
 }
 
-void AVirtualCameraActor::OnRemoteSessionChannelChange(IRemoteSessionRole* InRole, TWeakPtr<IRemoteSessionChannel> Channel, ERemoteSessionChannelChange Change)
-{
-	TSharedPtr<IRemoteSessionChannel> Pinned = Channel.Pin();
-
-	if (Pinned && Change == ERemoteSessionChannelChange::Created)
-	{
-		if (Pinned->GetType() == FRemoteSessionInputChannel::StaticType())
-		{
-			OnInputChannelCreated(Pinned);
-		}
-		else if (Pinned->GetType() == FRemoteSessionImageChannel::StaticType())
-		{
-			OnImageChannelCreated(Pinned);
-		}
-	}
-}
-
-void AVirtualCameraActor::OnImageChannelCreated(TWeakPtr<IRemoteSessionChannel> Instance)
-{
-	TSharedPtr<FRemoteSessionImageChannel> ImageChannel = StaticCastSharedPtr<FRemoteSessionImageChannel>(Instance.Pin());
-	if (ImageChannel)
-	{
-		ImageChannel->SetImageProvider(nullptr);
-		MediaOutput->SetImageChannel(ImageChannel);
-		MediaCapture = Cast<URemoteSessionMediaCapture>(MediaOutput->CreateMediaCapture());
-
-		TWeakPtr<SWindow> InputWindow;
-		TWeakPtr<FSceneViewport> SceneViewport;
-		FindSceneViewport(InputWindow, SceneViewport);
-		if (TSharedPtr<FSceneViewport> PinnedSceneViewport = SceneViewport.Pin())
-		{
-			MediaCapture->CaptureSceneViewport(PinnedSceneViewport, FMediaCaptureOptions());
-		}
-	}
-}
-
-void AVirtualCameraActor::OnInputChannelCreated(TWeakPtr<IRemoteSessionChannel> Instance)
-{
-	TSharedPtr<FRemoteSessionInputChannel> InputChannel = StaticCastSharedPtr<FRemoteSessionInputChannel>(Instance.Pin());
-	if (InputChannel)
-	{
-		TSharedPtr<SVirtualWindow> InputWindow = CameraScreenWidget->PostProcessDisplayType.GetSlateWindow();
-		InputChannel->SetPlaybackWindow(InputWindow, nullptr);
-		InputChannel->TryRouteTouchMessageToWidget(true);
-		InputChannel->GetOnRouteTouchDownToWidgetFailedDelegate()->AddUObject(this, &AVirtualCameraActor::OnTouchEventOutsideUMG);
-	}
-}
-
 void AVirtualCameraActor::OnTouchEventOutsideUMG(const FVector2D& InViewportPosition)
 {
 	FVector TraceDirection;
diff --git a/Engine/Plugins/Experimental/VirtualCamera/Source/VirtualCamera/Public/VirtualCameraActor.h b/Engine/Plugins/Experimental/VirtualCamera/Source/VirtualCamera/Public/VirtualCameraActor.h
--- a/Engine/Plugins/Experimental/VirtualCamera/Source/VirtualCamera/Public/VirtualCameraActor.h
+++ b/Engine/Plugins/Experimental/VirtualCamera/Source/VirtualCamera/Public/VirtualCameraActor.h
@@ -178,11 +178,6 @@ protected:
 	//~ End IVirtualCameraOptions Interface
 
 private:
-
-	void OnRemoteSessionChannelChange(IRemoteSessionRole* Role, TWeakPtr<IRemoteSessionChannel> Channel, ERemoteSessionChannelChange Change);
-	void OnImageChannelCreated(TWeakPtr<IRemoteSessionChannel> Instance);
-	void OnInputChannelCreated(TWeakPtr<IRemoteSessionChannel> Instance);
-
 	void OnTouchEventOutsideUMG(const FVector2D& InViewportPosition);
 
 	/** Stores the current camera settings to a save game for later use. */
@@ -203,7 +198,6 @@ private:
 	UPROPERTY(BlueprintGetter = GetCineCameraComponent, Category = "VirtualCamera | Component")
 	UCineCameraComponent* StreamedCamera;
 	bool bIsStreaming;
-	TSharedPtr<IRemoteSessionUnmanagedRole> RemoteSessionHost;
 	TUniquePtr<FVirtualCameraViewportSettings> ViewportSettingsBackup;
 	FHitResult LastViewportTouchResult;
 };
