Silence a mac warning

diff --git a/Engine/Plugins/Experimental/VirtualCamera/Source/VCamCore/Private/VCamComponent.cpp b/Engine/Plugins/Experimental/VirtualCamera/Source/VCamCore/Private/VCamComponent.cpp
--- a/Engine/Plugins/Experimental/VirtualCamera/Source/VCamCore/Private/VCamComponent.cpp
+++ b/Engine/Plugins/Experimental/VirtualCamera/Source/VCamCore/Private/VCamComponent.cpp
@@ -98,7 +98,7 @@ void UVCamComponent::OnComponentDestroyed(bool bDestroyingHierarchy)
 
 void UVCamComponent::HandleObjectReplaced(const TMap<UObject*, UObject*>& ReplacementMap)
 {
-	for (const TPair<UObject*, UObject*> ReplacementPair : ReplacementMap)
+	for (const TPair<UObject*, UObject*>& ReplacementPair : ReplacementMap)
 	{
 		UObject* FromObject = ReplacementPair.Key;
 		UObject* ToObject = ReplacementPair.Value;
