Do not log an error if we are being destroyed.

diff --git a/Engine/Plugins/Experimental/VirtualCamera/Source/VCamCore/Private/VCamComponent.cpp b/Engine/Plugins/Experimental/VirtualCamera/Source/VCamCore/Private/VCamComponent.cpp
--- a/Engine/Plugins/Experimental/VirtualCamera/Source/VCamCore/Private/VCamComponent.cpp
+++ b/Engine/Plugins/Experimental/VirtualCamera/Source/VCamCore/Private/VCamComponent.cpp
@@ -190,9 +190,12 @@ void UVCamComponent::CheckForErrors()
 {
 	Super::CheckForErrors();
 
-	if (!GetTargetCamera())
+	if (!HasAnyFlags(RF_BeginDestroyed|RF_FinishDestroyed))
 	{
-		UE_LOG(LogVCamComponent, Error, TEXT("Attached Parent should be a CineCamera derived component."));
+		if (!GetTargetCamera())
+		{
+			UE_LOG(LogVCamComponent, Error, TEXT("Attached Parent should be a CineCamera derived component."));
+		}
 	}
 }
 
