Fix non-unity build issue

diff --git a/Engine/Plugins/Experimental/VirtualCamera/Source/VirtualCamera/Private/VirtualCameraPlayerControllerBase.cpp b/Engine/Plugins/Experimental/VirtualCamera/Source/VirtualCamera/Private/VirtualCameraPlayerControllerBase.cpp
--- a/Engine/Plugins/Experimental/VirtualCamera/Source/VirtualCamera/Private/VirtualCameraPlayerControllerBase.cpp
+++ b/Engine/Plugins/Experimental/VirtualCamera/Source/VirtualCamera/Private/VirtualCameraPlayerControllerBase.cpp
@@ -6,6 +6,7 @@
 #include "CineCameraActor.h"
 #include "Engine/Engine.h"
 #include "Engine/World.h"
+#include "Engine/GameViewportClient.h"
 #include "Features/IModularFeatures.h"
 #include "Framework/Application/SlateApplication.h"
 #include "ILiveLinkClient.h"
