Fix include for -NoPCH builds

diff --git a/Engine/Plugins/Experimental/VirtualCamera/Source/VirtualCamera/Public/IVirtualCameraController.h b/Engine/Plugins/Experimental/VirtualCamera/Source/VirtualCamera/Public/IVirtualCameraController.h
--- a/Engine/Plugins/Experimental/VirtualCamera/Source/VirtualCamera/Public/IVirtualCameraController.h
+++ b/Engine/Plugins/Experimental/VirtualCamera/Source/VirtualCamera/Public/IVirtualCameraController.h
@@ -4,7 +4,7 @@
 
 #include "CoreMinimal.h"
 #include "Delegates/DelegateCombinations.h"
-#include "Engine/EngineTypes.h"
+#include "Engine/HitResult.h"
 #include "LiveLinkRole.h"
 #include "Misc/CoreMiscDefines.h"
 #include "UObject/Interface.h"
