Added extra guards to VCam Component to prevent Output Providers initializing across MU sessions

diff --git a/Engine/Plugins/Experimental/VirtualCamera/Source/VCamCore/Private/VCamComponent.cpp b/Engine/Plugins/Experimental/VirtualCamera/Source/VCamCore/Private/VCamComponent.cpp
--- a/Engine/Plugins/Experimental/VirtualCamera/Source/VCamCore/Private/VCamComponent.cpp
+++ b/Engine/Plugins/Experimental/VirtualCamera/Source/VCamCore/Private/VCamComponent.cpp
@@ -311,7 +311,7 @@ void UVCamComponent::PostEditChangeChainProperty(FPropertyChangedChainEvent& Pro
 							DestroyOutputProvider(SavedOutputProviders[ChangedIndex]);
 						}
 
-						if (ChangedProvider)
+						if (ChangedProvider && ShouldUpdateOutputProviders())
 						{
 							ChangedProvider->Initialize();
 						}
@@ -409,17 +409,20 @@ void UVCamComponent::Update()
 		SendCameraDataViaMultiUser();
 	}
 
-	for (UVCamOutputProviderBase* Provider : OutputProviders)
+	if (ShouldUpdateOutputProviders())
 	{
-		if (Provider)
+		for (UVCamOutputProviderBase* Provider : OutputProviders)
 		{
-			// Initialize the Provider if required
-			if (!Provider->IsInitialized())
+			if (Provider)
 			{
-				Provider->Initialize();
-			}
+				// Initialize the Provider if required
+				if (!Provider->IsInitialized())
+				{
+					Provider->Initialize();
+				}
 
-			Provider->Tick(DeltaTime);
+				Provider->Tick(DeltaTime);
+			}
 		}
 	}
 }
@@ -443,7 +446,7 @@ void UVCamComponent::SetEnabled(bool bNewEnabled)
 
 	// Enable any outputs that are set to active
 	// NOTE this must be done AFTER setting the actual bEnabled variable because OutputProviderBase now checks the component enabled state
-	if (bNewEnabled)
+	if (bNewEnabled && ShouldUpdateOutputProviders())
 	{
 		for (UVCamOutputProviderBase* Provider : OutputProviders)
 		{
@@ -910,7 +913,14 @@ void UVCamComponent::ResetAllOutputProviders()
 		{
 			// Initialization will also recover active state 
 			Provider->Deinitialize();
-			Provider->Initialize();
+
+			// We only Initialize a provider if they're able to be updated
+			// If they later become able to be updated then they will be
+			// Initialized inside the Update() loop
+			if (ShouldUpdateOutputProviders())
+			{
+				Provider->Initialize();
+			}
 		}
 	}
 }
@@ -1374,6 +1384,15 @@ bool UVCamComponent::CanEvaluateModifierStack() const
 	return !IsMultiUserSession() || (IsMultiUserSession() && IsCameraInVPRole());
 }
 
+bool UVCamComponent::ShouldUpdateOutputProviders() const
+{
+	// We should only update output providers in 3 situations
+	// - We're not in Multi User
+	// - We have the virtual camera role
+	// - We have explicitly set that we want to run Output Providers even when not in the camera role
+	return !IsMultiUserSession() || IsCameraInVPRole() || !bDisableOutputOnMultiUserReceiver;
+}
+
 bool UVCamComponent::IsMultiUserSession() const
 {
 #if WITH_EDITOR
diff --git a/Engine/Plugins/Experimental/VirtualCamera/Source/VCamCore/Public/VCamComponent.h b/Engine/Plugins/Experimental/VirtualCamera/Source/VCamCore/Public/VCamComponent.h
--- a/Engine/Plugins/Experimental/VirtualCamera/Source/VCamCore/Public/VCamComponent.h
+++ b/Engine/Plugins/Experimental/VirtualCamera/Source/VCamCore/Public/VCamComponent.h
@@ -304,6 +304,9 @@ private:
 	/** Can the modifier stack be evaluated. */
 	bool CanEvaluateModifierStack() const;
 
+	/** Output Providers should only update outside of MU or if we have the correct VP Role */
+	bool ShouldUpdateOutputProviders() const;
+
 	// When another component replaces us, get a notification so we can clean up
 	void NotifyComponentWasReplaced(UVCamComponent* ReplacementComponent);
 
