Made LiveLink a requirement of the VirtualCamera plugin

diff --git a/Engine/Plugins/Experimental/VirtualCamera/VirtualCamera.uplugin b/Engine/Plugins/Experimental/VirtualCamera/VirtualCamera.uplugin
--- a/Engine/Plugins/Experimental/VirtualCamera/VirtualCamera.uplugin
+++ b/Engine/Plugins/Experimental/VirtualCamera/VirtualCamera.uplugin
@@ -79,6 +79,10 @@
 		{
 			"Name": "VirtualProductionUtilities",
 			"Enabled": true
+		},
+		{
+			"Name": "LiveLink",
+			"Enabled": true
 		}
 	]
 }
