Prevent VCamComponent from trying to access the Live Link Client during cook

diff --git a/Engine/Plugins/Experimental/VirtualCamera/Source/VCamCore/Private/VCamComponent.cpp b/Engine/Plugins/Experimental/VirtualCamera/Source/VCamCore/Private/VCamComponent.cpp
--- a/Engine/Plugins/Experimental/VirtualCamera/Source/VCamCore/Private/VCamComponent.cpp
+++ b/Engine/Plugins/Experimental/VirtualCamera/Source/VCamCore/Private/VCamComponent.cpp
@@ -37,7 +37,7 @@ namespace VCamComponent
 UVCamComponent::UVCamComponent()
 {
 	// Don't run on CDO
-	if (!HasAnyFlags(RF_ClassDefaultObject))
+	if (!HasAnyFlags(RF_ClassDefaultObject) && !GIsCookerLoadingPackage)
 	{
 		// Hook into the Live Link Client for our Tick
 		IModularFeatures& ModularFeatures = IModularFeatures::Get();
