Initial support for Editor Input in VCam Modifiers

diff --git a/Engine/Plugins/Experimental/VirtualCamera/Source/VCamCore/Private/VCamComponent.cpp b/Engine/Plugins/Experimental/VirtualCamera/Source/VCamCore/Private/VCamComponent.cpp
--- a/Engine/Plugins/Experimental/VirtualCamera/Source/VCamCore/Private/VCamComponent.cpp
+++ b/Engine/Plugins/Experimental/VirtualCamera/Source/VCamCore/Private/VCamComponent.cpp
@@ -10,6 +10,9 @@
 #include "Roles/LiveLinkTransformRole.h"
 #include "GameDelegates.h"
 #include "Engine/GameEngine.h"
+#include "EnhancedInputComponent.h"
+#include "Engine/InputDelegateBinding.h"
+#include "InputMappingContext.h"
 
 #if WITH_EDITOR
 #include "Modules/ModuleManager.h"
@@ -24,6 +27,9 @@
 #include "IConcertSyncClient.h"
 #include "IMultiUserClientModule.h"
 
+#include "EnhancedInputEditorSubsystem.h"
+#include "GameFramework/InputSettings.h"
+
 #include "VPSettings.h"
 #endif
 
@@ -65,6 +71,25 @@ UVCamComponent::UVCamComponent()
 		MultiUserStartup();
 		FCoreUObjectDelegates::OnObjectsReplaced.AddUObject(this, &UVCamComponent::HandleObjectReplaced);
 #endif
+
+		// Setup Input
+		//EnhancedInputComponent = CreateDefaultSubobject<UEnhancedInputComponent>(TEXT("VCamInput"));
+		// From Ben H: Use NewObject here instead so that you can specify the class from the input settings :)
+		// Marking it as transient is also necessary
+		InputComponent = NewObject<UInputComponent>(this, UInputSettings::GetDefaultInputComponentClass(), TEXT("VCamInput0"), RF_Transient);
+		
+		
+#if WITH_EDITOR
+		if (GEditor && IsValid(InputComponent))
+		{
+			if (UEnhancedInputEditorSubsystem* Subsystem = GEditor->GetEditorSubsystem<UEnhancedInputEditorSubsystem>())
+			{
+				Subsystem->PushInputComponent(InputComponent);
+				Subsystem->StartConsumingInput();
+			}
+		}
+#endif
+		
 	}
 }
 
@@ -93,6 +118,14 @@ void UVCamComponent::OnComponentDestroyed(bool bDestroyingHierarchy)
 
 	MultiUserShutdown();
 	FCoreUObjectDelegates::OnObjectsReplaced.RemoveAll(this);
+	
+	if (GEditor && IsValid(InputComponent))
+	{
+		if (UEnhancedInputEditorSubsystem* Subsystem = GEditor->GetEditorSubsystem<UEnhancedInputEditorSubsystem>())
+		{
+			Subsystem->PopInputComponent(InputComponent);
+		}
+	}
 #endif
 }
 
@@ -311,6 +344,9 @@ void UVCamComponent::PostEditChangeChainProperty(FPropertyChangedChainEvent& Pro
 							DestroyOutputProvider(SavedOutputProviders[ChangedIndex]);
 						}
 
+						// We only Initialize a provider if they're able to be updated
+						// If they later become able to be updated then they will be
+						// Initialized inside the Update() loop
 						if (ChangedProvider && ShouldUpdateOutputProviders())
 						{
 							ChangedProvider->Initialize();
@@ -342,6 +378,25 @@ void UVCamComponent::PostEditChangeChainProperty(FPropertyChangedChainEvent& Pro
 }
 #endif // WITH_EDITOR
 
+void UVCamComponent::AddInputMappingContext(const UVCamModifier* Modifier)
+{
+	// Currently only handles the editor subsystem but made as a function so we can add game input at a later date 
+#if WITH_EDITOR
+	if (GEditor)
+	{
+		if (UEnhancedInputEditorSubsystem* Subsystem = GEditor->GetEditorSubsystem<UEnhancedInputEditorSubsystem>())
+		{
+			int32 InputPriority;
+			const UInputMappingContext* IMC = Modifier->GetInputMappingContext(InputPriority);
+			if (IsValid(IMC) && !Subsystem->HasMappingContext(IMC))
+			{
+				Subsystem->AddMappingContext(IMC, InputPriority);
+			}
+		}
+	}
+#endif
+}
+
 
 void UVCamComponent::Update()
 {
@@ -399,7 +454,8 @@ void UVCamComponent::Update()
 				// Initialize the Modifier if required
 				if (Modifier->DoesRequireInitialization())
 				{
-					Modifier->Initialize(ModifierContext);
+					Modifier->Initialize(ModifierContext, InputComponent);
+					AddInputMappingContext(Modifier);
 				}
 
 				Modifier->Apply(ModifierContext, CameraComponent, DeltaTime);
@@ -787,6 +843,14 @@ void UVCamComponent::GetLiveLinkDataForCurrentFrame(FLiveLinkCameraBlueprintData
 	}
 }
 
+void UVCamComponent::RegisterObjectForInput(UObject* Object)
+{
+	if (IsValid(InputComponent) && IsValid(Object))
+	{
+		UInputDelegateBinding::BindInputDelegates(GetClass(), InputComponent, this);
+	}
+}
+
 void UVCamComponent::CopyLiveLinkDataToCamera(const FLiveLinkCameraBlueprintData& LiveLinkData, UCineCameraComponent* CameraComponent)
 {
 	const FLiveLinkCameraStaticData& StaticData = LiveLinkData.StaticData;
diff --git a/Engine/Plugins/Experimental/VirtualCamera/Source/VCamCore/Private/VCamModifier.cpp b/Engine/Plugins/Experimental/VirtualCamera/Source/VCamCore/Private/VCamModifier.cpp
--- a/Engine/Plugins/Experimental/VirtualCamera/Source/VCamCore/Private/VCamModifier.cpp
+++ b/Engine/Plugins/Experimental/VirtualCamera/Source/VCamCore/Private/VCamModifier.cpp
@@ -3,9 +3,11 @@
 #include "VCamModifier.h"
 #include "VCamComponent.h"
 #include "VCamTypes.h"
+#include "EnhancedInputComponent.h"
+#include "Engine/InputDelegateBinding.h"
 
 
-void UVCamBlueprintModifier::Initialize(UVCamModifierContext* Context)
+void UVCamBlueprintModifier::Initialize(UVCamModifierContext* Context, UInputComponent* InputComponent)
 {
 	// Forward the Initialize call to the Blueprint Event
 	{
@@ -14,7 +16,7 @@ void UVCamBlueprintModifier::Initialize(UVCamModifierContext* Context)
 		OnInitialize(Context);
 	}
 
-	UVCamModifier::Initialize(Context);
+	UVCamModifier::Initialize(Context, InputComponent);
 }
 
 void UVCamBlueprintModifier::Apply(UVCamModifierContext* Context, UCineCameraComponent* CameraComponent, const float DeltaTime)
@@ -27,9 +29,27 @@ void UVCamBlueprintModifier::Apply(UVCamModifierContext* Context, UCineCameraCom
 	}
 }
 
-void UVCamModifier::Initialize(UVCamModifierContext* Context)
+const UInputMappingContext* UVCamBlueprintModifier::GetInputMappingContext(int32& InputPriority) const
 {
-	bRequiresInitialization = false; 
+	// Forward the call to the Blueprint implementation
+	{
+		FEditorScriptExecutionGuard ScriptGuard;
+
+		UInputMappingContext* MappingContext = nullptr;
+		GetInputMappingContextAndPriority(MappingContext, InputPriority);
+		return MappingContext;
+	}
+}
+
+void UVCamModifier::Initialize(UVCamModifierContext* Context, UInputComponent* InputComponent /* = nullptr */)
+{
+	// Binds any dynamic input delegates to the provided input component
+	if (IsValid(InputComponent))
+	{
+		UInputDelegateBinding::BindInputDelegates(GetClass(), InputComponent, this);
+	}
+
+	bRequiresInitialization = false;
 }
 
 void UVCamModifier::PostLoad()
@@ -70,6 +90,11 @@ bool UVCamModifier::IsEnabled() const
 	return false;
 }
 
+const UInputMappingContext* UVCamModifier::GetInputMappingContext(int32& InputPriority) const
+{
+	return nullptr;
+}
+
 FModifierStackEntry* UVCamModifier::GetCorrespondingStackEntry() const
 {
 	FModifierStackEntry* StackEntry = nullptr;
@@ -83,4 +108,4 @@ FModifierStackEntry* UVCamModifier::GetCorrespondingStackEntry() const
 	}
 
 	return StackEntry;
-}
\ No newline at end of file
+}
diff --git a/Engine/Plugins/Experimental/VirtualCamera/Source/VCamCore/Public/VCamComponent.h b/Engine/Plugins/Experimental/VirtualCamera/Source/VCamCore/Public/VCamComponent.h
--- a/Engine/Plugins/Experimental/VirtualCamera/Source/VCamCore/Public/VCamComponent.h
+++ b/Engine/Plugins/Experimental/VirtualCamera/Source/VCamCore/Public/VCamComponent.h
@@ -21,6 +21,7 @@ class UCineCameraComponent;
 class UVCamModifierContext;
 class SWindow;
 class FSceneViewport;
+class UEnhancedInputComponent;
 
 #if WITH_EDITOR
 class FLevelEditorViewportClient;
@@ -77,9 +78,12 @@ public:
 	void HandleObjectReplaced(const TMap<UObject*, UObject*>& ReplacementMap);
 
 	bool CanUpdate() const;
-
+	
 	void Update();
 
+	// Adds the Input Mapping Context from a modifier, if it exists, to the input system 
+	void AddInputMappingContext(const UVCamModifier* Modifier);
+
 	// Sets if the VCamComponent will update every frame or not
 	UFUNCTION(BlueprintSetter)
 	void SetEnabled(bool bNewEnabled);
@@ -247,6 +251,14 @@ public:
 	UFUNCTION(BlueprintCallable, Category="VirtualCamera")
 	void GetLiveLinkDataForCurrentFrame(FLiveLinkCameraBlueprintData& LiveLinkData);
 
+	/* Registers the given object with the VCamComponent's Input Component
+	 * This allows dynamic input bindings such as input events in blueprints to work correctly
+	 *
+	 * @param Object The object to register
+	 */
+	UFUNCTION(BlueprintCallable, Category="VirtualCamera")
+	void RegisterObjectForInput(UObject* Object);
+
 private:
 	static void CopyLiveLinkDataToCamera(const FLiveLinkCameraBlueprintData& LiveLinkData, UCineCameraComponent* CameraComponent);
 
@@ -331,4 +343,9 @@ private:
 
 	UPROPERTY(Transient)
 	bool bIsLockedToViewport = false;
+
+	// From Ben H: Mark this as Transient/DuplicateTransient so that it is saved on the BP CDO and nowhere else and
+	// handled correctly during duplication operations (copy/paste etc)
+	UPROPERTY(Transient, DuplicateTransient)
+	TObjectPtr<UInputComponent> InputComponent;
 };
diff --git a/Engine/Plugins/Experimental/VirtualCamera/Source/VCamCore/Public/VCamModifier.h b/Engine/Plugins/Experimental/VirtualCamera/Source/VCamCore/Public/VCamModifier.h
--- a/Engine/Plugins/Experimental/VirtualCamera/Source/VCamCore/Public/VCamModifier.h
+++ b/Engine/Plugins/Experimental/VirtualCamera/Source/VCamCore/Public/VCamModifier.h
@@ -8,8 +8,10 @@
 
 #include "VCamModifier.generated.h"
 
+class UInputMappingContext;
 class UVCamComponent;
 class UVCamModifierContext;
+class UInputComponent;
 
 struct FModifierStackEntry;
 
@@ -19,7 +21,7 @@ class VCAMCORE_API UVCamModifier : public UObject
 	GENERATED_BODY()
 
 public:
-	virtual void Initialize(UVCamModifierContext* Context);
+	virtual void Initialize(UVCamModifierContext* Context, UInputComponent* InputComponent = nullptr);
 
 	virtual void Apply(UVCamModifierContext* Context, UCineCameraComponent* CameraComponent, const float DeltaTime) {};
 
@@ -39,6 +41,10 @@ public:
 	UFUNCTION(BlueprintPure, Category = "VirtualCamera", meta=(ReturnDisplayName="Enabled"))
 	bool IsEnabled() const;
 
+	// Allows a modifier to return Input Mapping Context which will get automatically registered with the input system
+	// The Input Priority of the mapping context will be set by reference
+	virtual const UInputMappingContext* GetInputMappingContext(int32& InputPriority) const;
+
 private:
 	FModifierStackEntry* GetCorrespondingStackEntry() const;
 
@@ -51,12 +57,17 @@ class VCAMCORE_API UVCamBlueprintModifier : public UVCamModifier
 	GENERATED_BODY()
 
 public:
-	virtual void Initialize(UVCamModifierContext* Context) override;
+	virtual void Initialize(UVCamModifierContext* Context, UInputComponent* InputComponent=nullptr) override;
 	virtual void Apply(UVCamModifierContext* Context, UCineCameraComponent* CameraComponent, const float DeltaTime) override;
-
+	virtual const UInputMappingContext* GetInputMappingContext(int32& InputPriority) const override;
+	
 	UFUNCTION(BlueprintImplementableEvent, Category="VirtualCamera")
 	void OnInitialize(UVCamModifierContext* Context);
 
 	UFUNCTION(BlueprintImplementableEvent, Category="VirtualCamera")
 	void OnApply(UVCamModifierContext* Context, UCineCameraComponent* CameraComponent, const float DeltaTime);
+
+	// Allows a modifier to return Input Mapping Context which will get automatically registered with the input system
+	UFUNCTION(BlueprintImplementableEvent, Category="VirtualCamera")
+	void GetInputMappingContextAndPriority(UInputMappingContext*& InputMappingContext, int32& InputPriority) const;
 };
\ No newline at end of file
diff --git a/Engine/Plugins/Experimental/VirtualCamera/Source/VCamCore/VCamCore.Build.cs b/Engine/Plugins/Experimental/VirtualCamera/Source/VCamCore/VCamCore.Build.cs
--- a/Engine/Plugins/Experimental/VirtualCamera/Source/VCamCore/VCamCore.Build.cs
+++ b/Engine/Plugins/Experimental/VirtualCamera/Source/VCamCore/VCamCore.Build.cs
@@ -46,7 +46,8 @@ public class VCamCore : ModuleRules
 				"LiveLinkInterface",
 				"MediaIOCore",
 				"RemoteSession",
-				"InputCore"
+				"InputCore",
+				"EnhancedInput"
 				// ... add private dependencies that you statically link with here ...
 			}
 			);
@@ -69,7 +70,8 @@ public class VCamCore : ModuleRules
 					"UnrealEd",
 					"Concert",
 					"ConcertSyncClient",
-					"MultiUserClient"
+					"MultiUserClient",
+					"InputEditor"
 				}
 			);
 		}
diff --git a/Engine/Plugins/Experimental/VirtualCamera/VirtualCamera.uplugin b/Engine/Plugins/Experimental/VirtualCamera/VirtualCamera.uplugin
--- a/Engine/Plugins/Experimental/VirtualCamera/VirtualCamera.uplugin
+++ b/Engine/Plugins/Experimental/VirtualCamera/VirtualCamera.uplugin
@@ -83,6 +83,10 @@
 		{
 			"Name": "LiveLink",
 			"Enabled": true
+		},
+		{
+			"Name": "EnhancedInput",
+			"Enabled": true
 		}
 	]
 }
