Allow for 11ms updates in Virtual Camera.

diff --git a/Engine/Plugins/Experimental/VirtualCamera/Source/VCamCore/Public/VCamComponent.h b/Engine/Plugins/Experimental/VirtualCamera/Source/VCamCore/Public/VCamComponent.h
--- a/Engine/Plugins/Experimental/VirtualCamera/Source/VCamCore/Public/VCamComponent.h
+++ b/Engine/Plugins/Experimental/VirtualCamera/Source/VCamCore/Public/VCamComponent.h
@@ -236,8 +236,12 @@ public:
 	UPROPERTY(EditAnywhere, BlueprintReadWrite, AdvancedDisplay, Category = "VirtualCamera")
 	bool bDisableOutputOnMultiUserReceiver = true;
 
-	/** Indicates the frequency which camera updates are sent when in Multi-user mode. This has a minimum value of 30ms. */
-	UPROPERTY(EditAnywhere, AdvancedDisplay, Category="VirtualCamera", meta=(ForceUnits=ms, ClampMin = "30.0"), DisplayName="Update Frequencey")
+	/**
+	 * Indicates the frequency which camera updates are sent when in Multi-user mode. This has a minimum value of
+	 * 11ms. Using values below 30ms is discouraged. When higher refresh rates are needed consider using LiveLink
+	 * rebroadcast to stream camera data.
+	 */
+	UPROPERTY(EditAnywhere, AdvancedDisplay, Category="VirtualCamera", meta=(ForceUnits=ms, ClampMin = "11.0"), DisplayName="Update Frequencey")
 	float UpdateFrequencyMs = 66.6f;
 
 	// Which viewport to use for this VCam
