Allowed VCam Modifiers to access the associated name in the modifier stack

diff --git a/Engine/Plugins/Experimental/VirtualCamera/Source/VCamCore/Private/VCamModifier.cpp b/Engine/Plugins/Experimental/VirtualCamera/Source/VCamCore/Private/VCamModifier.cpp
--- a/Engine/Plugins/Experimental/VirtualCamera/Source/VCamCore/Private/VCamModifier.cpp
+++ b/Engine/Plugins/Experimental/VirtualCamera/Source/VCamCore/Private/VCamModifier.cpp
@@ -90,6 +90,36 @@ bool UVCamModifier::IsEnabled() const
 	return false;
 }
 
+bool UVCamModifier::SetStackEntryName(FName NewName)
+{
+	if (const UVCamComponent* ParentComponent = GetOwningVCamComponent())
+	{
+		const bool bNameAlreadyExists = ParentComponent->ModifierStack.ContainsByPredicate([NewName](const FModifierStackEntry& StackEntryToTest)
+		{
+			return StackEntryToTest.Name.IsEqual(NewName);
+		});
+
+		if (!bNameAlreadyExists)
+		{
+			if (FModifierStackEntry* StackEntry = GetCorrespondingStackEntry())
+			{
+				StackEntry->Name = NewName;
+				return true;
+			}
+		}
+	}
+	return false;
+}
+
+FName UVCamModifier::GetStackEntryName() const
+{
+	if (FModifierStackEntry* StackEntry = GetCorrespondingStackEntry())
+	{
+		return StackEntry->Name;
+	}
+	return NAME_None;
+}
+
 const UInputMappingContext* UVCamModifier::GetInputMappingContext(int32& InputPriority) const
 {
 	return nullptr;
diff --git a/Engine/Plugins/Experimental/VirtualCamera/Source/VCamCore/Public/VCamModifier.h b/Engine/Plugins/Experimental/VirtualCamera/Source/VCamCore/Public/VCamModifier.h
--- a/Engine/Plugins/Experimental/VirtualCamera/Source/VCamCore/Public/VCamModifier.h
+++ b/Engine/Plugins/Experimental/VirtualCamera/Source/VCamCore/Public/VCamModifier.h
@@ -41,6 +41,15 @@ public:
 	UFUNCTION(BlueprintPure, Category = "VirtualCamera", meta=(ReturnDisplayName="Enabled"))
 	bool IsEnabled() const;
 
+	// Sets the name of the modifier in the associated modifier stack
+	// Returns a bool for whether 
+	UFUNCTION(BlueprintCallable, Category = "VirtualCamera", meta=(ReturnDisplayName="Success"))
+	bool SetStackEntryName(FName NewName);
+
+	// Gets the name of the modifier in the associated modifier stack
+	UFUNCTION(BlueprintPure, Category = "VirtualCamera", meta=(ReturnDisplayName="Name"))
+	FName GetStackEntryName() const;
+
 	// Allows a modifier to return Input Mapping Context which will get automatically registered with the input system
 	// The Input Priority of the mapping context will be set by reference
 	virtual const UInputMappingContext* GetInputMappingContext(int32& InputPriority) const;
