Clean up some comments in the VCam component :)

#rnx

diff --git a/Engine/Plugins/Experimental/VirtualCamera/Source/VCamCore/Private/VCamComponent.cpp b/Engine/Plugins/Experimental/VirtualCamera/Source/VCamCore/Private/VCamComponent.cpp
--- a/Engine/Plugins/Experimental/VirtualCamera/Source/VCamCore/Private/VCamComponent.cpp
+++ b/Engine/Plugins/Experimental/VirtualCamera/Source/VCamCore/Private/VCamComponent.cpp
@@ -73,9 +73,6 @@ UVCamComponent::UVCamComponent()
 #endif
 
 		// Setup Input
-		//EnhancedInputComponent = CreateDefaultSubobject<UEnhancedInputComponent>(TEXT("VCamInput"));
-		// From Ben H: Use NewObject here instead so that you can specify the class from the input settings :)
-		// Marking it as transient is also necessary
 		InputComponent = NewObject<UInputComponent>(this, UInputSettings::GetDefaultInputComponentClass(), TEXT("VCamInput0"), RF_Transient);
 		
 		
