Added support for querying and modifying keybindings of input actions

diff --git a/Engine/Plugins/Experimental/VirtualCamera/Source/VCamCore/Private/VCamComponent.cpp b/Engine/Plugins/Experimental/VirtualCamera/Source/VCamCore/Private/VCamComponent.cpp
--- a/Engine/Plugins/Experimental/VirtualCamera/Source/VCamCore/Private/VCamComponent.cpp
+++ b/Engine/Plugins/Experimental/VirtualCamera/Source/VCamCore/Private/VCamComponent.cpp
@@ -159,6 +159,17 @@ void UVCamComponent::NotifyComponentWasReplaced(UVCamComponent* ReplacementCompo
 	DestroyComponent();
 }
 
+IEnhancedInputSubsystemInterface* UVCamComponent::GetEnhancedInputSubsystemInterface() const
+{
+#if WITH_EDITOR
+	if (GEditor)
+	{
+		return GEditor->GetEditorSubsystem<UEnhancedInputEditorSubsystem>();
+	}
+#endif
+	return nullptr;
+}
+
 bool UVCamComponent::CanUpdate() const
 {
 	UWorld* World = GetWorld();
@@ -377,21 +388,18 @@ void UVCamComponent::PostEditChangeChainProperty(FPropertyChangedChainEvent& Pro
 
 void UVCamComponent::AddInputMappingContext(const UVCamModifier* Modifier)
 {
-	// Currently only handles the editor subsystem but made as a function so we can add game input at a later date 
-#if WITH_EDITOR
-	if (GEditor)
+	if (IEnhancedInputSubsystemInterface* EnhancedInputSubsystemInterface = GetEnhancedInputSubsystemInterface())
 	{
-		if (UEnhancedInputEditorSubsystem* Subsystem = GEditor->GetEditorSubsystem<UEnhancedInputEditorSubsystem>())
+		int32 InputPriority;
+		const UInputMappingContext* IMC = Modifier->GetInputMappingContext(InputPriority);
+		if (IsValid(IMC))
 		{
-			int32 InputPriority;
-			const UInputMappingContext* IMC = Modifier->GetInputMappingContext(InputPriority);
-			if (IsValid(IMC) && !Subsystem->HasMappingContext(IMC))
+			if (!EnhancedInputSubsystemInterface->HasMappingContext(IMC))
 			{
-				Subsystem->AddMappingContext(IMC, InputPriority);
+				EnhancedInputSubsystemInterface->AddMappingContext(IMC, InputPriority);
 			}
 		}
 	}
-#endif
 }
 
 
@@ -844,8 +852,28 @@ void UVCamComponent::RegisterObjectForInput(UObject* Object)
 {
 	if (IsValid(InputComponent) && IsValid(Object))
 	{
-		UInputDelegateBinding::BindInputDelegates(GetClass(), InputComponent, this);
+		UInputDelegateBinding::BindInputDelegates(Object->GetClass(), InputComponent, Object);
+	}
+}
+
+TArray<FEnhancedActionKeyMapping> UVCamComponent::GetPlayerMappableKeys() const
+{
+	if (const IEnhancedInputSubsystemInterface* EnhancedInputSubsystemInterface = GetEnhancedInputSubsystemInterface())
+	{
+		return EnhancedInputSubsystemInterface->GetAllPlayerMappableActionKeyMappings();
 	}
+	return {};
+}
+
+int32 UVCamComponent::AddPlayerMappedKey(const FName MappingName, const FKey NewKey,
+	const FModifyContextOptions& Options) const
+{
+	if (IEnhancedInputSubsystemInterface* EnhancedInputSubsystemInterface = GetEnhancedInputSubsystemInterface())
+	{
+		return EnhancedInputSubsystemInterface->AddPlayerMappedKey(MappingName, NewKey, Options);
+	}
+
+	return 0;
 }
 
 void UVCamComponent::CopyLiveLinkDataToCamera(const FLiveLinkCameraBlueprintData& LiveLinkData, UCineCameraComponent* CameraComponent)
diff --git a/Engine/Plugins/Experimental/VirtualCamera/Source/VCamCore/Public/VCamComponent.h b/Engine/Plugins/Experimental/VirtualCamera/Source/VCamCore/Public/VCamComponent.h
--- a/Engine/Plugins/Experimental/VirtualCamera/Source/VCamCore/Public/VCamComponent.h
+++ b/Engine/Plugins/Experimental/VirtualCamera/Source/VCamCore/Public/VCamComponent.h
@@ -3,6 +3,8 @@
 #pragma once
 
 #include "CoreMinimal.h"
+#include "EnhancedActionKeyMapping.h"
+#include "EnhancedInputSubsystemInterface.h"
 #include "VCamTypes.h"
 #include "Roles/LiveLinkCameraTypes.h"
 #include "VCamOutputProviderBase.h"
@@ -263,6 +265,22 @@ public:
 	UFUNCTION(BlueprintCallable, Category="VirtualCamera")
 	void RegisterObjectForInput(UObject* Object);
 
+	/**
+	 * Returns a list of all player mappable keys that have been registered
+	 */
+	UFUNCTION(BlueprintCallable, Category="VirtualCamera")
+	TArray<FEnhancedActionKeyMapping> GetPlayerMappableKeys() const;
+
+	/**
+	 * Replace any currently applied mappings to this key mapping with the given new one.
+	 * Requests a rebuild of the player mappings.
+	 * This function will automatically switch between Editor based and Player based input systems as needed
+	 *
+	 * @return The number of mappings that have been replaced
+	 */
+	UFUNCTION(BlueprintCallable, Category = "VirtualCamera", meta=(AutoCreateRefTerm = "Options"))
+	int32 AddPlayerMappedKey(const FName MappingName, const FKey NewKey, const FModifyContextOptions& Options = FModifyContextOptions()) const;
+
 private:
 	static void CopyLiveLinkDataToCamera(const FLiveLinkCameraBlueprintData& LiveLinkData, UCineCameraComponent* CameraComponent);
 
@@ -326,6 +344,12 @@ private:
 	// When another component replaces us, get a notification so we can clean up
 	void NotifyComponentWasReplaced(UVCamComponent* ReplacementComponent);
 
+	/*
+	 * Gets the input interface for the currently active input method.
+	 * Will switch between editor and player based input as needed
+	 */
+	IEnhancedInputSubsystemInterface* GetEnhancedInputSubsystemInterface() const;
+
 	double LastEvaluationTime;
 
 	TWeakObjectPtr<AActor> Backup_ActorLock;
diff --git a/Engine/Plugins/Experimental/VirtualCamera/Source/VCamCore/VCamCore.Build.cs b/Engine/Plugins/Experimental/VirtualCamera/Source/VCamCore/VCamCore.Build.cs
--- a/Engine/Plugins/Experimental/VirtualCamera/Source/VCamCore/VCamCore.Build.cs
+++ b/Engine/Plugins/Experimental/VirtualCamera/Source/VCamCore/VCamCore.Build.cs
@@ -27,7 +27,8 @@ public class VCamCore : ModuleRules
 			{
 				"Composure",
 				"Core",
-				"VPUtilities"
+				"VPUtilities",
+				"EnhancedInput"
 				// ... add other public dependencies that you statically link with here ...
 			}
 			);
@@ -46,8 +47,7 @@ public class VCamCore : ModuleRules
 				"LiveLinkInterface",
 				"MediaIOCore",
 				"RemoteSession",
-				"InputCore",
-				"EnhancedInput"
+				"InputCore"
 				// ... add private dependencies that you statically link with here ...
 			}
 			);
