Ported and modified Lyra AnyKey Widget for use in VCam

diff --git a/Engine/Plugins/Experimental/VirtualCamera/Source/VCamInput/Private/VCamAnyKeyWidget.cpp b/Engine/Plugins/Experimental/VirtualCamera/Source/VCamInput/Private/VCamAnyKeyWidget.cpp
new file mode 100644
--- /dev/null
+++ b/Engine/Plugins/Experimental/VirtualCamera/Source/VCamInput/Private/VCamAnyKeyWidget.cpp
@@ -0,0 +1,137 @@
+// Copyright Epic Games, Inc. All Rights Reserved.
+
+#include "VCamAnyKeyWidget.h"
+#include "Framework/Application/IInputProcessor.h"
+
+// Input processor class that forwards any relevant input to the owning widget
+// All Handle events return true to consume the input
+class FVCamPressAnyKeyInputProcessor : public IInputProcessor
+{
+public:
+	FVCamPressAnyKeyInputProcessor(UVCamPressAnyKey* InOwner) : Owner(InOwner) {};
+
+	virtual void Tick(const float DeltaTime, FSlateApplication& SlateApp, TSharedRef<ICursor> Cursor) override {};
+
+	virtual bool HandleKeyUpEvent(FSlateApplication& SlateApp, const FKeyEvent& InKeyEvent) override
+	{
+		HandleKey(InKeyEvent.GetKey());
+		return true;
+	}
+
+	virtual bool HandleKeyDownEvent(FSlateApplication& SlateApp, const FKeyEvent& InKeyEvent) override
+	{
+		return true;
+	}
+
+	virtual bool HandleMouseButtonUpEvent(FSlateApplication& SlateApp, const FPointerEvent& MouseEvent) override
+	{
+		HandleKey(MouseEvent.GetEffectingButton());
+		return true;
+	}
+
+	virtual bool HandleMouseButtonDownEvent(FSlateApplication& SlateApp, const FPointerEvent& MouseEvent) override
+	{
+		return true;
+	}
+
+	virtual bool HandleMouseButtonDoubleClickEvent(FSlateApplication& SlateApp, const FPointerEvent& MouseEvent) override
+	{
+		HandleKey(MouseEvent.GetEffectingButton());
+		return true;
+	}
+
+	virtual bool HandleMouseWheelOrGestureEvent(FSlateApplication& SlateApp, const FPointerEvent& InWheelEvent, const FPointerEvent* InGestureEvent) override
+	{
+		if (InWheelEvent.GetWheelDelta() != 0)
+		{
+			const FKey Key = InWheelEvent.GetWheelDelta() < 0 ? EKeys::MouseScrollDown : EKeys::MouseScrollUp;
+			HandleKey(Key);
+		}
+		return true;
+	}
+
+private:
+	void HandleKey(const FKey& Key) const
+	{
+		// We should always have a valid owner as the owner keeps us alive via a shared pointer
+		check(IsValid(Owner));
+
+		// Cancel the input if the key is escape or a touch event, otherwise pass the key back to the owner
+		if (Key == EKeys::LeftCommand || Key == EKeys::RightCommand)
+		{
+			// Ignore
+		}
+		else if (Key == EKeys::Escape || Key.IsTouch())
+		{
+			Owner->HandleKeySelectionCanceled();
+		}
+		else
+		{
+			Owner->HandleKeySelected(Key);
+		}
+	}
+
+	
+	// Keep a reference to the widget owning this processor so we can notify on key events
+	UVCamPressAnyKey* Owner = nullptr;
+};
+
+void UVCamPressAnyKey::NativeOnActivated()
+{
+	Super::NativeOnActivated();
+
+	bKeySelected = false;
+
+	InputProcessor = MakeShared<FVCamPressAnyKeyInputProcessor>(this);
+	FSlateApplication::Get().RegisterInputPreProcessor(InputProcessor, 0);
+}
+
+void UVCamPressAnyKey::NativeOnDeactivated()
+{
+	Super::NativeOnDeactivated();
+
+	if (FSlateApplication::IsInitialized())
+	{
+		FSlateApplication::Get().UnregisterInputPreProcessor(InputProcessor);
+	}
+}
+
+void UVCamPressAnyKey::HandleKeySelected(FKey InKey)
+{
+	if (!bKeySelected)
+	{
+		bKeySelected = true;
+		Dismiss([this, InKey]()
+		{
+			OnKeySelected.Broadcast(InKey);
+		});
+	}
+}
+
+void UVCamPressAnyKey::HandleKeySelectionCanceled()
+{
+if (!bKeySelected)
+	{
+		bKeySelected = true;
+		Dismiss([this]()
+		{
+			OnKeySelectionCanceled.Broadcast();
+		});
+	}
+}
+
+void UVCamPressAnyKey::Dismiss(TFunction<void()> PostDismissCallback)
+{
+	FTSTicker::GetCoreTicker().AddTicker(FTickerDelegate::CreateWeakLambda(this, [this, PostDismissCallback](float DeltaTime)
+	{
+		if (FSlateApplication::IsInitialized())
+		{
+			FSlateApplication::Get().UnregisterInputPreProcessor(InputProcessor);
+		}
+
+		DeactivateWidget();
+		PostDismissCallback();
+
+		return false;
+	}));
+}
diff --git a/Engine/Plugins/Experimental/VirtualCamera/Source/VCamInput/Public/VCamAnyKeyWidget.h b/Engine/Plugins/Experimental/VirtualCamera/Source/VCamInput/Public/VCamAnyKeyWidget.h
new file mode 100644
--- /dev/null
+++ b/Engine/Plugins/Experimental/VirtualCamera/Source/VCamInput/Public/VCamAnyKeyWidget.h
@@ -0,0 +1,38 @@
+// Copyright Epic Games, Inc. All Rights Reserved.
+
+#pragma once
+
+#include "CoreMinimal.h"
+#include "CommonActivatableWidget.h"
+#include "VCamAnyKeyWidget.generated.h"
+
+DECLARE_DYNAMIC_MULTICAST_DELEGATE_OneParam(FOnKeySelected, FKey, Key);
+DECLARE_DYNAMIC_MULTICAST_DELEGATE(FOnKeySelectionCanceled);
+
+// Base widget class that allows you to wait for any input and then respond to the key pressed
+UCLASS(Abstract)
+class VCAMINPUT_API UVCamPressAnyKey : public UCommonActivatableWidget
+{
+	GENERATED_BODY()
+public:
+	UPROPERTY(BlueprintAssignable)
+	FOnKeySelected OnKeySelected;
+
+	UPROPERTY(BlueprintAssignable)
+	FOnKeySelectionCanceled OnKeySelectionCanceled;
+
+protected:
+	virtual void NativeOnActivated() override;
+	virtual void NativeOnDeactivated() override;
+
+	void HandleKeySelected(FKey InKey);
+	void HandleKeySelectionCanceled();
+
+	void Dismiss(TFunction<void()> PostDismissCallback);
+
+private:
+	bool bKeySelected = false;
+	TSharedPtr<class FVCamPressAnyKeyInputProcessor> InputProcessor;
+
+	friend class FVCamPressAnyKeyInputProcessor;
+};
\ No newline at end of file
diff --git a/Engine/Plugins/Experimental/VirtualCamera/Source/VCamInput/VCamInput.build.cs b/Engine/Plugins/Experimental/VirtualCamera/Source/VCamInput/VCamInput.build.cs
--- a/Engine/Plugins/Experimental/VirtualCamera/Source/VCamInput/VCamInput.build.cs
+++ b/Engine/Plugins/Experimental/VirtualCamera/Source/VCamInput/VCamInput.build.cs
@@ -13,6 +13,8 @@ public class VCamInput : ModuleRules
 			new string[]
 			{
 				"Core",
+				"UMG",
+				"CommonUI"
 			}
 			);
 
diff --git a/Engine/Plugins/Experimental/VirtualCamera/VirtualCamera.uplugin b/Engine/Plugins/Experimental/VirtualCamera/VirtualCamera.uplugin
--- a/Engine/Plugins/Experimental/VirtualCamera/VirtualCamera.uplugin
+++ b/Engine/Plugins/Experimental/VirtualCamera/VirtualCamera.uplugin
@@ -87,6 +87,10 @@
 		{
 			"Name": "EnhancedInput",
 			"Enabled": true
+		},
+		{
+			"Name": "CommonUI",
+			"Enabled": true
 		}
 	]
 }
