Added the notion of a "VCam Connection" to allow for a safer mechanism for UI to interact with Modifiers

diff --git a/Engine/Plugins/Experimental/VirtualCamera/Source/VCamCore/Private/UI/VCamConnectionStructs.cpp b/Engine/Plugins/Experimental/VirtualCamera/Source/VCamCore/Private/UI/VCamConnectionStructs.cpp
new file mode 100644
--- /dev/null
+++ b/Engine/Plugins/Experimental/VirtualCamera/Source/VCamCore/Private/UI/VCamConnectionStructs.cpp
@@ -0,0 +1,68 @@
+// Copyright Epic Games, Inc. All Rights Reserved.
+
+#include "UI/VCamConnectionStructs.h"
+
+#include "Algo/AllOf.h"
+#include "VCamComponent.h"
+
+DEFINE_LOG_CATEGORY(LogVCamConnection);
+
+bool FVCamConnection::IsConnected() const
+{
+	// If we have a valid modifier then we can assume we have a successful connection
+	return IsValid(ConnectedModifier);
+}
+
+bool FVCamConnection::AttemptConnection(UVCamComponent* VCamComponent)
+{
+	bool bIsValidModifierConnection = false;
+		UVCamModifier* Modifier = VCamComponent->GetModifierByName(ConnectionTargetSettings.TargetModifierName);
+		FVCamModifierConnectionPoint* ConnectionPoint = nullptr;
+		if (Modifier)
+		{			
+			const bool bImplementsAllRequiredInterfaces = Algo::AllOf(RequiredInterfaces, [Modifier](const TSubclassOf<UInterface>& Interface)
+			{
+				return Modifier->GetClass()->ImplementsInterface(Interface);
+			});
+
+			if (bImplementsAllRequiredInterfaces)
+			{
+				ConnectionPoint = Modifier->ConnectionPoints.Find(ConnectionTargetSettings.TargetConnectionPoint);
+				if (ConnectionPoint)
+				{
+					// The Connection Point is considered valid if either we don't require an action or the required ActionType matches the ActionType inside the Connection Point
+					const bool bConnectionPointHasRequiredActionType = !bRequiresInputAction || (ConnectionPoint->AssociatedAction && ConnectionPoint->AssociatedAction->ValueType == ActionType);
+					if (bConnectionPointHasRequiredActionType)
+					{
+						bIsValidModifierConnection = true;
+					}
+					else
+					{
+						UE_LOG(LogVCamConnection, Warning, TEXT("Connection Failed: Connection Point with name %s does not provide an action of the required type. Expected %s but found %s"),
+							*ConnectionTargetSettings.TargetConnectionPoint.ToString(),
+							*UEnum::GetValueAsString(ActionType),
+							ConnectionPoint->AssociatedAction ? *UEnum::GetValueAsString(ConnectionPoint->AssociatedAction->ValueType) : TEXT("None"));
+					}
+				}
+				else
+				{
+					UE_LOG(LogVCamConnection, Warning, TEXT("Connection Failed: Failed to find Connection Point with name %s in Modifier %s"),
+						*ConnectionTargetSettings.TargetConnectionPoint.ToString(),
+						*ConnectionTargetSettings.TargetModifierName.ToString());
+				}
+			}
+			else
+			{
+				UE_LOG(LogVCamConnection, Warning, TEXT("Connection Failed: Modifier with name %s does not implement all required interfaces for this connection."), *ConnectionTargetSettings.TargetModifierName.ToString());
+			}
+		}
+		else
+		{
+			UE_LOG(LogVCamConnection, Warning, TEXT("Connection Failed: Failed to find Modifier with name %s in VCam Component"), *ConnectionTargetSettings.TargetModifierName.ToString());
+		}
+
+		ConnectedModifier = bIsValidModifierConnection ? Modifier : nullptr;
+		ConnectedAction = bIsValidModifierConnection ? ConnectionPoint->AssociatedAction : nullptr;
+
+	return bIsValidModifierConnection;
+}
diff --git a/Engine/Plugins/Experimental/VirtualCamera/Source/VCamCore/Private/UI/VCamUIFunctionLibrary.cpp b/Engine/Plugins/Experimental/VirtualCamera/Source/VCamCore/Private/UI/VCamUIFunctionLibrary.cpp
new file mode 100644
--- /dev/null
+++ b/Engine/Plugins/Experimental/VirtualCamera/Source/VCamCore/Private/UI/VCamUIFunctionLibrary.cpp
@@ -0,0 +1,25 @@
+// Copyright Epic Games, Inc. All Rights Reserved.
+
+#include "UI/VCamUIFunctionLibrary.h"
+
+bool UVCamUIFunctionLibrary::IsConnected_VCamConnection(const FVCamConnection& Connection)
+{
+	return Connection.IsConnected();
+}
+
+FName UVCamUIFunctionLibrary::GetConnectionPointName_VCamConnection(const FVCamConnection& Connection)
+{
+	return Connection.ConnectionTargetSettings.TargetConnectionPoint;
+}
+
+UVCamModifier* UVCamUIFunctionLibrary::GetConnectedModifier_VCamConnection(const FVCamConnection& Connection)
+{
+	return Connection.ConnectedModifier;
+}
+
+UInputAction* UVCamUIFunctionLibrary::GetConnectedInputAction_VCamConnection(const FVCamConnection& Connection)
+{
+	return Connection.ConnectedAction;
+}
+
+
diff --git a/Engine/Plugins/Experimental/VirtualCamera/Source/VCamCore/Private/UI/VCamWidget.cpp b/Engine/Plugins/Experimental/VirtualCamera/Source/VCamCore/Private/UI/VCamWidget.cpp
new file mode 100644
--- /dev/null
+++ b/Engine/Plugins/Experimental/VirtualCamera/Source/VCamCore/Private/UI/VCamWidget.cpp
@@ -0,0 +1,27 @@
+// Copyright Epic Games, Inc. All Rights Reserved.
+
+#include "UI/VCamWidget.h"
+
+void UVCamWidget::InitializeConnections(UVCamComponent* VCam)
+{
+	if (!IsValid(VCam))
+	{
+		return;
+	}
+
+	// Iteratively call AttemptConnection on each connection within the widget and notify the result via OnConnectionUpdated
+	for (TPair<FName, FVCamConnection>& Connection : Connections)
+	{
+		const FName& ConnectionName = Connection.Key;
+		FVCamConnection& VCamConnection = Connection.Value;
+
+		const bool bDidConnectSuccessfully = VCamConnection.AttemptConnection(VCam);
+
+		if (!bDidConnectSuccessfully)
+		{
+			UE_LOG(LogVCamConnection, Warning, TEXT("Widget %s: Failed to create for VCam Connection with Connection Name: %s"), *GetName(), *ConnectionName.ToString());	
+		}
+		
+		OnConnectionUpdated(ConnectionName, bDidConnectSuccessfully, VCamConnection.ConnectionTargetSettings.TargetConnectionPoint, VCamConnection.ConnectedModifier, VCamConnection.ConnectedAction);
+	}
+}
diff --git a/Engine/Plugins/Experimental/VirtualCamera/Source/VCamCore/Private/VCamOutputProviderBase.cpp b/Engine/Plugins/Experimental/VirtualCamera/Source/VCamCore/Private/VCamOutputProviderBase.cpp
--- a/Engine/Plugins/Experimental/VirtualCamera/Source/VCamCore/Private/VCamOutputProviderBase.cpp
+++ b/Engine/Plugins/Experimental/VirtualCamera/Source/VCamCore/Private/VCamOutputProviderBase.cpp
@@ -6,6 +6,8 @@
 #include "VCamComponent.h"
 #include "UObject/UObjectBaseUtility.h"
 #include "Blueprint/UserWidget.h"
+#include "Blueprint/WidgetTree.h"
+#include "UI/VCamWidget.h"
 
 #if WITH_EDITOR
 #include "Editor.h"
@@ -213,13 +215,28 @@ void UVCamOutputProviderBase::NotifyWidgetOfComponentChange() const
 	if (UMGWidget && UMGWidget->IsDisplayed())
 	{
 		UUserWidget* DisplayedWidget = UMGWidget->GetWidget();
-		if (DisplayedWidget && DisplayedWidget->Implements<UVCamModifierInterface>())
+		if (IsValid(DisplayedWidget))
 		{
-			if (UVCamComponent* OwningComponent = Cast<UVCamComponent>(this->GetOuter()))
+			if (UVCamComponent* OwningComponent = GetTypedOuter<UVCamComponent>())
 			{
-				UVCamComponent* CameraComponent = bIsActive ? OwningComponent : nullptr;
+				UVCamComponent* VCamComponent = bIsActive ? OwningComponent : nullptr;
 
-				IVCamModifierInterface::Execute_OnVCamComponentChanged(DisplayedWidget, CameraComponent);
+				if (DisplayedWidget->Implements<UVCamModifierInterface>())
+				{
+					IVCamModifierInterface::Execute_OnVCamComponentChanged(DisplayedWidget, VCamComponent);
+				}
+
+				// Find all VCam Widgets inside the displayed widget and Initialize them with the owning VCam Component
+				if (IsValid(DisplayedWidget->WidgetTree))
+				{
+					DisplayedWidget->WidgetTree->ForEachWidget([VCamComponent](UWidget* Widget)
+					{
+						if (UVCamWidget* VCamWidget = Cast<UVCamWidget>(Widget))
+						{
+							VCamWidget->InitializeConnections(VCamComponent);
+						}
+					});
+				}
 			}
 		}
 
diff --git a/Engine/Plugins/Experimental/VirtualCamera/Source/VCamCore/Public/UI/VCamConnectionStructs.h b/Engine/Plugins/Experimental/VirtualCamera/Source/VCamCore/Public/UI/VCamConnectionStructs.h
new file mode 100644
--- /dev/null
+++ b/Engine/Plugins/Experimental/VirtualCamera/Source/VCamCore/Public/UI/VCamConnectionStructs.h
@@ -0,0 +1,116 @@
+// Copyright Epic Games, Inc. All Rights Reserved.
+
+#pragma once
+
+#include "CoreMinimal.h"
+#include "InputActionValue.h"
+#include "Templates/SubclassOf.h"
+
+#include "VCamConnectionStructs.generated.h"
+
+DECLARE_LOG_CATEGORY_EXTERN(LogVCamConnection, Log, All);
+
+class UInputAction;
+class UInterface;
+class UVCamComponent;
+class UVCamModifier;
+
+/*
+ * This stores information about the target for a given VCam Widget Connection
+ * Ideally this is controlled externally via a configurator system but can be set explicitly if required
+ */
+USTRUCT(BlueprintType)
+struct VCAMCORE_API FVCamConnectionTargetSettings
+{
+	GENERATED_BODY()
+
+	/*
+	 * The name of the Modifier to search for in a given VCam Component's Modifier Stack
+	 */
+	UPROPERTY(EditAnywhere, BlueprintReadWrite, Category=ConnectionBinding)
+	FName TargetModifierName = NAME_None;
+
+	/*
+	 * The name of the Connection Point to search for in a given VCam Modifier
+	 */
+	UPROPERTY(EditAnywhere, BlueprintReadWrite, Category=ConnectionBinding)
+	FName TargetConnectionPoint = NAME_None;
+};
+
+/*
+ * A VCam Connection allows an external object (primarily widgets) to communicate with VCam Modifiers through Connection Points that have known properties.
+ * The connection allows you to specify a set of requirements that a connection point on a modifier must implement
+ * for it to be considered a valid thing to connect to and then a function "AttemptConnection" that will try to resolve
+ * a given VCam Component to a specific Modifier along with an optional associated Input Action
+ */
+USTRUCT(BlueprintType)
+struct VCAMCORE_API FVCamConnection
+{
+	GENERATED_BODY()
+
+	/*
+	 * A list of interfaces that a modifier must implement to be considered a valid connection
+	 */
+	UPROPERTY(EditAnywhere, BlueprintReadWrite, Category=Connection)
+	TArray<TSubclassOf<UInterface>> RequiredInterfaces;
+
+	/*
+	 * A list of interfaces that a modifier may optionally implement that this connection can use
+	 *
+	 * These interfaces are not tested for when attempting a connection but are stored so that
+	 * a Configurator system can provide extra information a user when they are configuring connections
+	 */
+	UPROPERTY(EditAnywhere, BlueprintReadWrite, Category=Connection)
+	TArray<TSubclassOf<UInterface>> OptionalInterfaces;
+
+	/*
+	 * Whether this connection requires a target Connection Point to have an associated Input Action to be considered a valid connection
+	 */
+	UPROPERTY(EditAnywhere, BlueprintReadWrite, Category=Connection)
+	bool bRequiresInputAction = true;
+
+	/*
+	 * If this connection requires an Input Action then this property specifies what Input Action Type the Connection Point must provide to be considered a valid connection
+	 */
+	UPROPERTY(EditAnywhere, BlueprintReadWrite, Category=Connection, meta=(EditCondition="bRequiresInputAction"))
+	EInputActionValueType ActionType;
+
+	/*
+	 * If you are not using an external system to configure the connection then you can enable this to have
+	 * explicit control over which Modifier and Connection Point to look for when attempting a connection
+	 */
+	UPROPERTY(EditAnywhere, BlueprintReadWrite, Category=Connection)
+	bool bManuallyConfigureConnection;
+
+	/*
+	 * A struct containing information about which Modifier and Connection Point to look for when attempting a connection
+	 */
+	UPROPERTY(EditAnywhere, BlueprintReadWrite, Category=Connection, meta=(EditCondition="bManuallyConfigureConnection", EditConditionHides))
+	FVCamConnectionTargetSettings ConnectionTargetSettings;
+
+	/*
+	 * The Input Action that we are currently connected to
+	 *
+	 * This could be empty even if we have a valid connection
+	 */
+	UPROPERTY(Transient)
+	TObjectPtr<UInputAction> ConnectedAction = nullptr;
+
+	/*
+	 * The VCam Modifier that we are currently connected to
+	 *
+	 * Guaranteed to be valid if IsConnected() is true
+	 */
+	UPROPERTY(Transient)
+	TObjectPtr<UVCamModifier> ConnectedModifier = nullptr;
+
+	/*
+	 * Check whether this connection was successfully connected to a Connection Point
+	 */
+	bool IsConnected() const;
+
+	/*
+	 * Attempts to make a connection to a target Connection Point and returns a bool indicating success
+	 */
+	bool AttemptConnection(UVCamComponent* VCamComponent);
+};
diff --git a/Engine/Plugins/Experimental/VirtualCamera/Source/VCamCore/Public/UI/VCamUIFunctionLibrary.h b/Engine/Plugins/Experimental/VirtualCamera/Source/VCamCore/Public/UI/VCamUIFunctionLibrary.h
new file mode 100644
--- /dev/null
+++ b/Engine/Plugins/Experimental/VirtualCamera/Source/VCamCore/Public/UI/VCamUIFunctionLibrary.h
@@ -0,0 +1,52 @@
+// Copyright Epic Games, Inc. All Rights Reserved.
+
+#pragma once
+
+#include "CoreMinimal.h"
+#include "Kismet/BlueprintFunctionLibrary.h"
+#include "UI/VCamConnectionStructs.h"
+
+#include "VCamUIFunctionLibrary.generated.h"
+
+class UInputAction;
+class UVCamModifier;
+
+/*
+ * Function Library primarily used for exposing functionality of UI related structs to Blueprints 
+ */
+UCLASS()
+class VCAMCORE_API UVCamUIFunctionLibrary : public UBlueprintFunctionLibrary
+{
+	GENERATED_BODY()
+
+public:
+	/*
+	 * Checks whether a given VCam Connection is successfully connected.
+	 * If this returns true then it is safe to assume that the Connected Modifier for the Connection is valid to use
+	 */
+	UFUNCTION(BlueprintPure, Category = "VCam Connections", meta=(DisplayName = "Is Connected"))
+	static bool IsConnected_VCamConnection(const FVCamConnection& Connection);
+
+	/*
+	 * Gets the name of the associated Connection Point for a given VCam Connection.
+	 * This will return "None" if the IsConnected() is false
+	 */
+	UFUNCTION(BlueprintPure, Category = "VCam Connections", meta = (DisplayName = "Get Connection Point Name"))
+	static FName GetConnectionPointName_VCamConnection(const FVCamConnection& Connection);
+
+	/*
+	 * Gets the Connected Modifier for the given VCam Connection
+	 * The returned Modifier is guaranteed to be valid if IsConnected() is true
+	 */
+	UFUNCTION(BlueprintPure, Category = "VCam Connections", meta = (DisplayName = "Get Connected Modifier"))
+	static UVCamModifier* GetConnectedModifier_VCamConnection(const FVCamConnection& Connection);
+
+	/*
+	 * Gets the Connected Input Action for the given VCam Connection
+	 * The returned Input Action is only guaranteed to be valid if both IsConnected() is true AND the VCam Connection is set to require an Input Action
+	 * If the VCam Connection doesn't require an Input Action then this may still return a valid Action if the Connection Point
+	 * had an associated Input Action however you will need to test for this manually
+	 */
+	UFUNCTION(BlueprintPure, Category = "VCam Connections", meta = (DisplayName = "Get Connected Input Action"))
+	static UInputAction* GetConnectedInputAction_VCamConnection(const FVCamConnection& Connection);
+};
\ No newline at end of file
diff --git a/Engine/Plugins/Experimental/VirtualCamera/Source/VCamCore/Public/UI/VCamWidget.h b/Engine/Plugins/Experimental/VirtualCamera/Source/VCamCore/Public/UI/VCamWidget.h
new file mode 100644
--- /dev/null
+++ b/Engine/Plugins/Experimental/VirtualCamera/Source/VCamCore/Public/UI/VCamWidget.h
@@ -0,0 +1,48 @@
+// Copyright Epic Games, Inc. All Rights Reserved.
+
+#pragma once
+
+#include "CoreMinimal.h"
+#include "VCamComponent.h"
+#include "UI/VCamConnectionStructs.h"
+#include "Blueprint/UserWidget.h"
+
+#include "VCamWidget.generated.h"
+
+class UInputAction;
+class UVCamModifier;
+
+/*
+ * A wrapper widget class that contains a set of VCam Connections
+ * 
+ * If you add a widget deriving from UVCamWidget to an Overlay Widget for a VCam Output Provider then when the
+ * Overlay is created by the Provider it will also call InitializeConnections with the owning VCam Component.
+ */
+UCLASS(Abstract)
+class VCAMCORE_API UVCamWidget : public UUserWidget
+{
+	GENERATED_BODY()
+public:
+	/*
+	 * The VCam Connections associated with this Widget
+	 * 
+	 * Each Connection has a unique name associated with it and any connection related event
+	 * will provide this name as one of its arguments.
+	 */
+	UPROPERTY(EditAnywhere, BlueprintReadWrite, Category="VCam Connections")
+	TMap<FName, FVCamConnection> Connections;
+
+	/*
+	 * Event called when a specific connection has been updated
+	 * 
+	 * The connection is not guaranteed to succeed so "Did Connect Successfully" should be checked before using
+	 * the connected modifier or action
+	 */
+	UFUNCTION(BlueprintImplementableEvent, Category="VCam Connections")
+	void OnConnectionUpdated(FName ConnectionName, bool bDidConnectSuccessfully, FName ModifierConnectionPointName, UVCamModifier* ConnectedModifier, UInputAction* ConnectedAction);
+
+	/*
+	 * Iterate all VCam Connections within the widget and attempt to connect them using the provided VCam Component
+	 */
+	void InitializeConnections(UVCamComponent* VCam);
+};
\ No newline at end of file
diff --git a/Engine/Plugins/Experimental/VirtualCamera/Source/VCamCore/Public/VCamModifier.h b/Engine/Plugins/Experimental/VirtualCamera/Source/VCamCore/Public/VCamModifier.h
--- a/Engine/Plugins/Experimental/VirtualCamera/Source/VCamCore/Public/VCamModifier.h
+++ b/Engine/Plugins/Experimental/VirtualCamera/Source/VCamCore/Public/VCamModifier.h
@@ -8,6 +8,7 @@
 
 #include "VCamModifier.generated.h"
 
+class UInputAction;
 class UInputMappingContext;
 class UVCamComponent;
 class UVCamModifierContext;
@@ -15,6 +16,16 @@ class UInputComponent;
 
 struct FModifierStackEntry;
 
+USTRUCT(BlueprintType)
+struct VCAMCORE_API FVCamModifierConnectionPoint
+{
+	GENERATED_BODY()
+
+	// An optional action to associate with this connection point
+	UPROPERTY(EditAnywhere, BlueprintReadWrite, Category = "Connection")
+	TObjectPtr<UInputAction> AssociatedAction;
+};
+
 UCLASS(Blueprintable, Abstract, EditInlineNew)
 class VCAMCORE_API UVCamModifier : public UObject
 {
@@ -54,6 +65,9 @@ public:
 	// The Input Priority of the mapping context will be set by reference
 	virtual const UInputMappingContext* GetInputMappingContext(int32& InputPriority) const;
 
+	UPROPERTY(EditAnywhere, BlueprintReadWrite, Category = "VCam Connection Points")
+	TMap<FName, FVCamModifierConnectionPoint> ConnectionPoints;
+	
 private:
 	FModifierStackEntry* GetCorrespondingStackEntry() const;
 
