Deprecating old VCam BPs with CoreRedirects

diff --git a/Engine/Plugins/Experimental/VirtualCamera/Config/BaseEngine.ini b/Engine/Plugins/Experimental/VirtualCamera/Config/BaseEngine.ini
--- a/Engine/Plugins/Experimental/VirtualCamera/Config/BaseEngine.ini
+++ b/Engine/Plugins/Experimental/VirtualCamera/Config/BaseEngine.ini
@@ -4,4 +4,9 @@
 +FunctionRedirects=(OldName="VCamComponent.ClearModifiers",NewName="VCamComponent.RemoveAllModifiers")
 +FunctionRedirects=(OldName="VCamComponent.FindModifierByName",NewName="VCamComponent.GetModifierByName")
 +FunctionRedirects=(OldName="VCamComponent.FindModifiersByClass",NewName="VCamComponent.GetModifiersByClass")
-+FunctionRedirects=(OldName="VCamComponent.FindModifierByInterface",NewName="VCamComponent.GetModifiersByInterface")
\ No newline at end of file
++FunctionRedirects=(OldName="VCamComponent.FindModifierByInterface",NewName="VCamComponent.GetModifiersByInterface")
++PackageRedirects=(OldName="/VirtualCamera/Editor", NewName="/VirtualCamera/Deprecated/Editor", MatchSubstring=true)
++PackageRedirects=(OldName="/VirtualCamera/OperatorPanel", NewName="/VirtualCamera/Deprecated/OperatorPanel", MatchSubstring=true)
++PackageRedirects=(OldName="/VirtualCamera/TestMaps", NewName="/VirtualCamera/Deprecated/TestMaps", MatchSubstring=true)
++PackageRedirects=(OldName="/VirtualCamera/V2", NewName="/VirtualCamera/Deprecated/V2", MatchSubstring=true)
++PackageRedirects=(OldName="/VirtualCamera/VCamCore", NewName="/VirtualCamera/Deprecated/VCamCore", MatchSubstring=true)
diff --git a/Engine/Plugins/Experimental/VirtualCamera/Config/BaseVirtualCamera.ini b/Engine/Plugins/Experimental/VirtualCamera/Config/BaseVirtualCamera.ini
new file mode 100644
--- /dev/null
+++ b/Engine/Plugins/Experimental/VirtualCamera/Config/BaseVirtualCamera.ini
@@ -0,0 +1,5 @@
+[CoreRedirects]
++PackageRedirects=(OldName="/VirtualCamera/OperatorPanel", NewName="/VirtualCamera/Deprecated/OperatorPanel", MatchSubstring=true)
++PackageRedirects=(OldName="/VirtualCamera/TestMaps", NewName="/VirtualCamera/Deprecated/TestMaps", MatchSubstring=true)
++PackageRedirects=(OldName="/VirtualCamera/V2", NewName="/VirtualCamera/Deprecated/V2", MatchSubstring=true)
++PackageRedirects=(OldName="/VirtualCamera/VCamCore", NewName="/VirtualCamera/Deprecated/VCamCore", MatchSubstring=true)
\ No newline at end of file
