Make sure we always notify the UI of any changes with the VCam Component which may require the connections to be updated

diff --git a/Engine/Plugins/Experimental/VirtualCamera/Source/VCamCore/Private/VCamComponent.cpp b/Engine/Plugins/Experimental/VirtualCamera/Source/VCamCore/Private/VCamComponent.cpp
--- a/Engine/Plugins/Experimental/VirtualCamera/Source/VCamCore/Private/VCamComponent.cpp
+++ b/Engine/Plugins/Experimental/VirtualCamera/Source/VCamCore/Private/VCamComponent.cpp
@@ -322,6 +322,14 @@ void UVCamComponent::PostEditChangeProperty(FPropertyChangedEvent& PropertyChang
 				}
 			}
 		}
+
+		for (const UVCamOutputProviderBase* OutputProvider : OutputProviders)
+		{
+			if (IsValid(OutputProvider))
+			{
+				OutputProvider->NotifyWidgetOfComponentChange();
+			}
+		}
 	}
 
 	Super::PostEditChangeProperty(PropertyChangedEvent);
diff --git a/Engine/Plugins/Experimental/VirtualCamera/Source/VCamCore/Public/VCamOutputProviderBase.h b/Engine/Plugins/Experimental/VirtualCamera/Source/VCamCore/Public/VCamOutputProviderBase.h
--- a/Engine/Plugins/Experimental/VirtualCamera/Source/VCamCore/Public/VCamOutputProviderBase.h
+++ b/Engine/Plugins/Experimental/VirtualCamera/Source/VCamCore/Public/VCamOutputProviderBase.h
@@ -106,6 +106,9 @@ public:
 		bWasActive = false;
 	}
 
+	/** Calls the VCamModifierInterface on the widget if it exists and also requests any child VCam Widgets to reconnect */
+	void NotifyWidgetOfComponentChange() const;
+
 protected:
 	// If set, this output provider will execute every frame
 	UPROPERTY(EditAnywhere, BlueprintReadWrite, Category = "Output", meta = (DisplayPriority = "1"))
@@ -135,8 +138,6 @@ protected:
 	UVPFullScreenUserWidget* UMGWidget = nullptr;
 
 private:
-	void NotifyWidgetOfComponentChange() const;
-
 	bool IsOuterComponentEnabled() const;
 
 	TSoftObjectPtr<UCineCameraComponent> TargetCamera;
