Re-enable VCamCore as a runtime plugin. All editor only uses should be
guarded.

diff --git a/Engine/Plugins/Experimental/VirtualCamera/VirtualCamera.uplugin b/Engine/Plugins/Experimental/VirtualCamera/VirtualCamera.uplugin
--- a/Engine/Plugins/Experimental/VirtualCamera/VirtualCamera.uplugin
+++ b/Engine/Plugins/Experimental/VirtualCamera/VirtualCamera.uplugin
@@ -17,7 +17,7 @@
 	"Modules": [
 		{
 			"Name": "VCamCore",
-			"Type": "UncookedOnly",
+			"Type": "Runtime",
 			"LoadingPhase": "Default",
 			"PlatformDenyList": [
 				"TVOS"
