Adding InjectInputForAction function on VCamComponent

diff --git a/Engine/Plugins/Experimental/VirtualCamera/Source/VCamCore/Private/VCamComponent.cpp b/Engine/Plugins/Experimental/VirtualCamera/Source/VCamCore/Private/VCamComponent.cpp
--- a/Engine/Plugins/Experimental/VirtualCamera/Source/VCamCore/Private/VCamComponent.cpp
+++ b/Engine/Plugins/Experimental/VirtualCamera/Source/VCamCore/Private/VCamComponent.cpp
@@ -886,6 +886,14 @@ int32 UVCamComponent::AddPlayerMappedKey(const FName MappingName, const FKey New
 	return 0;
 }
 
+void UVCamComponent::InjectInputForAction(const UInputAction* Action, FInputActionValue RawValue, const TArray<UInputModifier*>& Modifiers, const TArray<UInputTrigger*>& Triggers)
+{
+	if (IEnhancedInputSubsystemInterface* EnhancedInputSubsystemInterface = GetEnhancedInputSubsystemInterface())
+	{
+		EnhancedInputSubsystemInterface->InjectInputForAction( Action, RawValue, Modifiers, Triggers);
+	}
+}
+
 void UVCamComponent::CopyLiveLinkDataToCamera(const FLiveLinkCameraBlueprintData& LiveLinkData, UCineCameraComponent* CameraComponent)
 {
 	const FLiveLinkCameraStaticData& StaticData = LiveLinkData.StaticData;
diff --git a/Engine/Plugins/Experimental/VirtualCamera/Source/VCamCore/Public/VCamComponent.h b/Engine/Plugins/Experimental/VirtualCamera/Source/VCamCore/Public/VCamComponent.h
--- a/Engine/Plugins/Experimental/VirtualCamera/Source/VCamCore/Public/VCamComponent.h
+++ b/Engine/Plugins/Experimental/VirtualCamera/Source/VCamCore/Public/VCamComponent.h
@@ -281,6 +281,12 @@ public:
 	UFUNCTION(BlueprintCallable, Category = "VirtualCamera", meta=(AutoCreateRefTerm = "Options"))
 	int32 AddPlayerMappedKey(const FName MappingName, const FKey NewKey, const FModifyContextOptions& Options = FModifyContextOptions()) const;
 
+	/**
+	* Injects an input action. 
+	*/
+	UFUNCTION(BlueprintCallable, Category = "VirtualCamera", meta = (AutoCreateRefTerm = "Modifiers,Triggers"))
+	virtual void InjectInputForAction(const UInputAction* Action, FInputActionValue RawValue, const TArray<UInputModifier*>& Modifiers, const TArray<UInputTrigger*>& Triggers);
+
 private:
 	static void CopyLiveLinkDataToCamera(const FLiveLinkCameraBlueprintData& LiveLinkData, UCineCameraComponent* CameraComponent);
 
