Implement GetWorld in VCamModifier to expose some bp functions

diff --git a/Engine/Plugins/Experimental/VirtualCamera/Source/VCamCore/Private/VCamModifier.cpp b/Engine/Plugins/Experimental/VirtualCamera/Source/VCamCore/Private/VCamModifier.cpp
--- a/Engine/Plugins/Experimental/VirtualCamera/Source/VCamCore/Private/VCamModifier.cpp
+++ b/Engine/Plugins/Experimental/VirtualCamera/Source/VCamCore/Private/VCamModifier.cpp
@@ -139,3 +139,12 @@ FModifierStackEntry* UVCamModifier::GetCorrespondingStackEntry() const
 
 	return StackEntry;
 }
+
+UWorld* UVCamModifier::GetWorld() const
+{
+	if (UVCamComponent* ParentComponent = GetOwningVCamComponent())
+	{
+		return ParentComponent->GetWorld();
+	}
+	return nullptr;
+}
\ No newline at end of file
diff --git a/Engine/Plugins/Experimental/VirtualCamera/Source/VCamCore/Public/VCamModifier.h b/Engine/Plugins/Experimental/VirtualCamera/Source/VCamCore/Public/VCamModifier.h
--- a/Engine/Plugins/Experimental/VirtualCamera/Source/VCamCore/Public/VCamModifier.h
+++ b/Engine/Plugins/Experimental/VirtualCamera/Source/VCamCore/Public/VCamModifier.h
@@ -67,6 +67,8 @@ public:
 
 	UPROPERTY(EditAnywhere, BlueprintReadWrite, Category = "VCam Connection Points")
 	TMap<FName, FVCamModifierConnectionPoint> ConnectionPoints;
+
+	virtual UWorld* GetWorld() const override;
 	
 private:
 	FModifierStackEntry* GetCorrespondingStackEntry() const;
