VCam Widgets can now opt into receiving input events from the VCam Component via a set of properties on the widget. These properties control whether the widget will get registered for input (defaults to true) and which Input Mapping Context to add to the input system when registered. The bool controlling registration is only read when initializing so toggling the bool at runtime will not have an effect

diff --git a/Engine/Plugins/Experimental/VirtualCamera/Source/VCamCore/Private/UI/VCamWidget.cpp b/Engine/Plugins/Experimental/VirtualCamera/Source/VCamCore/Private/UI/VCamWidget.cpp
--- a/Engine/Plugins/Experimental/VirtualCamera/Source/VCamCore/Private/UI/VCamWidget.cpp
+++ b/Engine/Plugins/Experimental/VirtualCamera/Source/VCamCore/Private/UI/VCamWidget.cpp
@@ -24,4 +24,11 @@ void UVCamWidget::InitializeConnections(UVCamComponent* VCam)
 		
 		OnConnectionUpdated(ConnectionName, bDidConnectSuccessfully, VCamConnection.ConnectionTargetSettings.TargetConnectionPoint, VCamConnection.ConnectedModifier, VCamConnection.ConnectedAction);
 	}
+
+	// Register for input with the VCam Component if desired
+	if (bRegisterForInput)
+	{
+		VCam->RegisterObjectForInput(this);
+		VCam->AddInputMappingContext(InputMappingContext, InputContextPriority);
+	}
 }
diff --git a/Engine/Plugins/Experimental/VirtualCamera/Source/VCamCore/Private/VCamComponent.cpp b/Engine/Plugins/Experimental/VirtualCamera/Source/VCamCore/Private/VCamComponent.cpp
--- a/Engine/Plugins/Experimental/VirtualCamera/Source/VCamCore/Private/VCamComponent.cpp
+++ b/Engine/Plugins/Experimental/VirtualCamera/Source/VCamCore/Private/VCamComponent.cpp
@@ -412,6 +412,17 @@ void UVCamComponent::AddInputMappingContext(const UVCamModifier* Modifier)
 	}
 }
 
+void UVCamComponent::AddInputMappingContext(UInputMappingContext* Context, int32 Priority)
+{
+	if (IEnhancedInputSubsystemInterface* EnhancedInputSubsystemInterface = GetEnhancedInputSubsystemInterface())
+	{
+		if (IsValid(Context))
+		{
+			EnhancedInputSubsystemInterface->AddMappingContext(Context, Priority);
+		}
+	}
+}
+
 
 void UVCamComponent::Update()
 {
diff --git a/Engine/Plugins/Experimental/VirtualCamera/Source/VCamCore/Public/UI/VCamWidget.h b/Engine/Plugins/Experimental/VirtualCamera/Source/VCamCore/Public/UI/VCamWidget.h
--- a/Engine/Plugins/Experimental/VirtualCamera/Source/VCamCore/Public/UI/VCamWidget.h
+++ b/Engine/Plugins/Experimental/VirtualCamera/Source/VCamCore/Public/UI/VCamWidget.h
@@ -32,6 +32,26 @@ public:
 	UPROPERTY(EditAnywhere, BlueprintReadWrite, Category="VCam Connections")
 	TMap<FName, FVCamConnection> Connections;
 
+	/*
+	 * Determines whether this widget will be automatically registered to receive input when the connections are initialized
+	 *
+	 * Note: This property is only read during Initialize so toggling at runtime will not have any effect
+	 */
+	UPROPERTY(EditAnywhere, BlueprintReadOnly, Category="VCam Input")
+	bool bRegisterForInput = true;
+
+	/*
+	 * If this widget is registered for input then this input mapping context will be added to the input system
+	 */
+	UPROPERTY(EditAnywhere, BlueprintReadOnly, Category="VCam Input", meta=(EditCondition="bRegisterForInput"))
+	TObjectPtr<UInputMappingContext> InputMappingContext = nullptr;
+
+	/*
+	 * If this widget is registered for input then this property defines the priority that the input mapping context is added at
+	 */
+	UPROPERTY(EditAnywhere, BlueprintReadOnly, Category="VCam Input", meta=(EditCondition="bRegisterForInput"))
+	int32 InputContextPriority = 0;
+	
 	/*
 	 * Event called when a specific connection has been updated
 	 * 
diff --git a/Engine/Plugins/Experimental/VirtualCamera/Source/VCamCore/Public/VCamComponent.h b/Engine/Plugins/Experimental/VirtualCamera/Source/VCamCore/Public/VCamComponent.h
--- a/Engine/Plugins/Experimental/VirtualCamera/Source/VCamCore/Public/VCamComponent.h
+++ b/Engine/Plugins/Experimental/VirtualCamera/Source/VCamCore/Public/VCamComponent.h
@@ -86,6 +86,9 @@ public:
 	// Adds the Input Mapping Context from a modifier, if it exists, to the input system 
 	void AddInputMappingContext(const UVCamModifier* Modifier);
 
+	// Adds and explicitly provided Input Mapping Context to the input system
+	void AddInputMappingContext(UInputMappingContext* Context, int32 Priority);
+
 	// Sets if the VCamComponent will update every frame or not
 	UFUNCTION(BlueprintSetter)
 	void SetEnabled(bool bNewEnabled);
