Added functionality to copy a CineCameraComponent from a vcam onto a regular CineCameraActor

diff --git a/Engine/Plugins/Experimental/VirtualCamera/Source/VirtualCamera/Private/VCamBlueprintFunctionLibrary.cpp b/Engine/Plugins/Experimental/VirtualCamera/Source/VirtualCamera/Private/VCamBlueprintFunctionLibrary.cpp
--- a/Engine/Plugins/Experimental/VirtualCamera/Source/VirtualCamera/Private/VCamBlueprintFunctionLibrary.cpp
+++ b/Engine/Plugins/Experimental/VirtualCamera/Source/VirtualCamera/Private/VCamBlueprintFunctionLibrary.cpp
@@ -1,6 +1,8 @@
 // Copyright Epic Games, Inc. All Rights Reserved.
 
 #include "VCamBlueprintFunctionLibrary.h"
+
+#include "CineCameraActor.h"
 #include "CineCameraComponent.h"
 #include "LevelSequence.h"
 #include "VirtualCameraClipsMetaData.h"
@@ -400,6 +402,28 @@ FString UVCamBlueprintFunctionLibrary::GetNextUndoDescription()
 	return "";
 }
 
+bool UVCamBlueprintFunctionLibrary::CopyToCineCameraActor(UCineCameraComponent* SourceCameraComponent, ACineCameraActor* TargetCameraActor)
+{
+	if (!(IsValid(SourceCameraComponent) && IsValid(TargetCameraActor) && IsValid(TargetCameraActor->GetCineCameraComponent())))
+	{
+		return false;
+	}
+
+	// Ensure the root actor transforms match
+	if (const USceneComponent* SourceAttachment = SourceCameraComponent->GetAttachParent())
+	{
+		TargetCameraActor->SetActorTransform(SourceAttachment->GetComponentTransform());
+	}
+
+	// Copy the properties from one CineCameraComponent to the other
+	UCineCameraComponent* TargetCameraComponent = TargetCameraActor->GetCineCameraComponent();
+	UEngine::CopyPropertiesForUnrelatedObjects(SourceCameraComponent, TargetCameraComponent);
+	
+	TargetCameraActor->UpdateComponentTransforms();
+
+	return true;
+}
+
 bool UVCamBlueprintFunctionLibrary::DeprojectScreenToWorld(const FVector2D& InScreenPosition, FVector& OutWorldPosition, FVector& OutWorldDirection)
 {
 	FName LevelEditorName(TEXT("LevelEditor"));
diff --git a/Engine/Plugins/Experimental/VirtualCamera/Source/VirtualCamera/Public/VCamBlueprintFunctionLibrary.h b/Engine/Plugins/Experimental/VirtualCamera/Source/VirtualCamera/Public/VCamBlueprintFunctionLibrary.h
--- a/Engine/Plugins/Experimental/VirtualCamera/Source/VirtualCamera/Public/VCamBlueprintFunctionLibrary.h
+++ b/Engine/Plugins/Experimental/VirtualCamera/Source/VirtualCamera/Public/VCamBlueprintFunctionLibrary.h
@@ -2,6 +2,7 @@
 
 #pragma once
 
+#include "CineCameraComponent.h"
 #include "AssetRegistry/AssetData.h"
 #include "Kismet/BlueprintFunctionLibrary.h"
 #include "Misc/FrameRate.h"
@@ -9,6 +10,7 @@
 #include "MovieSceneObjectBindingID.h"
 #include "VCamBlueprintFunctionLibrary.generated.h"
 
+class ACineCameraActor;
 class UCineCameraComponent;
 class ULevelSequence;
 class USceneCaptureComponent2D;
@@ -141,6 +143,10 @@ public:
 	UFUNCTION(BlueprintPure, Category = "VirtualCamera")
 	static FString GetNextUndoDescription();
 
+	/** Copies all properties from a CineCameraComponent to a CineCameraActor and ensure the root actor transform is updated so the CameraComponents end up in the same World Space position */
+	UFUNCTION(BlueprintCallable, Category = "VirtualCamera")
+	static bool CopyToCineCameraActor(UCineCameraComponent* SourceCameraComponent, ACineCameraActor* TargetCameraActor);
+
 private:
 
 	static bool DeprojectScreenToWorld(const FVector2D& InScreenPosition, FVector& OutWorldPosition, FVector& OutWorldDirection);
