Allowed VCam Widgets to be initialized from blueprints and ensured nested VCam Widgets will correctly initialize

diff --git a/Engine/Plugins/Experimental/VirtualCamera/Source/VCamCore/Private/VCamOutputProviderBase.cpp b/Engine/Plugins/Experimental/VirtualCamera/Source/VCamCore/Private/VCamOutputProviderBase.cpp
--- a/Engine/Plugins/Experimental/VirtualCamera/Source/VCamCore/Private/VCamOutputProviderBase.cpp
+++ b/Engine/Plugins/Experimental/VirtualCamera/Source/VCamCore/Private/VCamOutputProviderBase.cpp
@@ -229,7 +229,7 @@ void UVCamOutputProviderBase::NotifyWidgetOfComponentChange() const
 				// Find all VCam Widgets inside the displayed widget and Initialize them with the owning VCam Component
 				if (IsValid(DisplayedWidget->WidgetTree))
 				{
-					DisplayedWidget->WidgetTree->ForEachWidget([VCamComponent](UWidget* Widget)
+					DisplayedWidget->WidgetTree->ForEachWidgetAndDescendants([VCamComponent](UWidget* Widget)
 					{
 						if (UVCamWidget* VCamWidget = Cast<UVCamWidget>(Widget))
 						{
diff --git a/Engine/Plugins/Experimental/VirtualCamera/Source/VCamCore/Public/UI/VCamWidget.h b/Engine/Plugins/Experimental/VirtualCamera/Source/VCamCore/Public/UI/VCamWidget.h
--- a/Engine/Plugins/Experimental/VirtualCamera/Source/VCamCore/Public/UI/VCamWidget.h
+++ b/Engine/Plugins/Experimental/VirtualCamera/Source/VCamCore/Public/UI/VCamWidget.h
@@ -64,5 +64,6 @@ public:
 	/*
 	 * Iterate all VCam Connections within the widget and attempt to connect them using the provided VCam Component
 	 */
+	UFUNCTION(BlueprintCallable, Category="VCam Connections")
 	void InitializeConnections(UVCamComponent* VCam);
 };
\ No newline at end of file
