Add InjectInputVectorForAction

diff --git a/Engine/Plugins/Experimental/VirtualCamera/Source/VCamCore/Private/VCamComponent.cpp b/Engine/Plugins/Experimental/VirtualCamera/Source/VCamCore/Private/VCamComponent.cpp
--- a/Engine/Plugins/Experimental/VirtualCamera/Source/VCamCore/Private/VCamComponent.cpp
+++ b/Engine/Plugins/Experimental/VirtualCamera/Source/VCamCore/Private/VCamComponent.cpp
@@ -906,6 +906,14 @@ void UVCamComponent::InjectInputForAction(const UInputAction* Action, FInputActi
 	}
 }
 
+void UVCamComponent::InjectInputVectorForAction(const UInputAction* Action, FVector Value, const TArray<UInputModifier*>& Modifiers, const TArray<UInputTrigger*>& Triggers)
+{
+	if (IEnhancedInputSubsystemInterface* EnhancedInputSubsystemInterface = GetEnhancedInputSubsystemInterface())
+	{
+		EnhancedInputSubsystemInterface->InjectInputVectorForAction(Action, Value, Modifiers, Triggers);
+	}
+}
+
 void UVCamComponent::CopyLiveLinkDataToCamera(const FLiveLinkCameraBlueprintData& LiveLinkData, UCineCameraComponent* CameraComponent)
 {
 	const FLiveLinkCameraStaticData& StaticData = LiveLinkData.StaticData;
diff --git a/Engine/Plugins/Experimental/VirtualCamera/Source/VCamCore/Public/VCamComponent.h b/Engine/Plugins/Experimental/VirtualCamera/Source/VCamCore/Public/VCamComponent.h
--- a/Engine/Plugins/Experimental/VirtualCamera/Source/VCamCore/Public/VCamComponent.h
+++ b/Engine/Plugins/Experimental/VirtualCamera/Source/VCamCore/Public/VCamComponent.h
@@ -290,6 +290,13 @@ public:
 	UFUNCTION(BlueprintCallable, Category = "VirtualCamera", meta = (AutoCreateRefTerm = "Modifiers,Triggers"))
 	virtual void InjectInputForAction(const UInputAction* Action, FInputActionValue RawValue, const TArray<UInputModifier*>& Modifiers, const TArray<UInputTrigger*>& Triggers);
 
+
+	/**
+	* Injects an input vector for action.
+	*/
+	UFUNCTION(BlueprintCallable, Category = "VirtualCamera", meta = (AutoCreateRefTerm = "Modifiers,Triggers"))
+	virtual void InjectInputVectorForAction(const UInputAction* Action, FVector Value, const TArray<UInputModifier*>& Modifiers, const TArray<UInputTrigger*>& Triggers);
+
 private:
 	static void CopyLiveLinkDataToCamera(const FLiveLinkCameraBlueprintData& LiveLinkData, UCineCameraComponent* CameraComponent);
 
