- VCam Widgets now cache the VCam Component that initialized them
- Added a Utility function to easily update connection targets and reinitialize connections on VCam Widgets
- Removed warnings when restting connections to None

diff --git a/Engine/Plugins/Experimental/VirtualCamera/Source/VCamCore/Private/UI/VCamConnectionStructs.cpp b/Engine/Plugins/Experimental/VirtualCamera/Source/VCamCore/Private/UI/VCamConnectionStructs.cpp
--- a/Engine/Plugins/Experimental/VirtualCamera/Source/VCamCore/Private/UI/VCamConnectionStructs.cpp
+++ b/Engine/Plugins/Experimental/VirtualCamera/Source/VCamCore/Private/UI/VCamConnectionStructs.cpp
@@ -66,3 +66,9 @@ bool FVCamConnection::AttemptConnection(UVCamComponent* VCamComponent)
 
 	return bIsValidModifierConnection;
 }
+
+void FVCamConnection::ResetConnection()
+{
+	ConnectedModifier = nullptr;
+	ConnectedAction = nullptr;
+}
diff --git a/Engine/Plugins/Experimental/VirtualCamera/Source/VCamCore/Private/UI/VCamWidget.cpp b/Engine/Plugins/Experimental/VirtualCamera/Source/VCamCore/Private/UI/VCamWidget.cpp
--- a/Engine/Plugins/Experimental/VirtualCamera/Source/VCamCore/Private/UI/VCamWidget.cpp
+++ b/Engine/Plugins/Experimental/VirtualCamera/Source/VCamCore/Private/UI/VCamWidget.cpp
@@ -4,31 +4,82 @@
 
 void UVCamWidget::InitializeConnections(UVCamComponent* VCam)
 {
-	if (!IsValid(VCam))
+	VCamComponent = VCam;
+
+	const bool bDidSuccessfullyInitialize = ReinitializeConnections();
+	
+	// Register for input with the VCam Component if desired
+	if (bDidSuccessfullyInitialize && bRegisterForInput)
 	{
-		return;
+		VCamComponent->RegisterObjectForInput(this);
+		VCamComponent->AddInputMappingContext(InputMappingContext, InputContextPriority);
 	}
+}
 
+bool UVCamWidget::ReinitializeConnections()
+{
+	if (!IsValid(VCamComponent))
+	{
+		return false;
+	}
+	
 	// Iteratively call AttemptConnection on each connection within the widget and notify the result via OnConnectionUpdated
 	for (TPair<FName, FVCamConnection>& Connection : Connections)
 	{
 		const FName& ConnectionName = Connection.Key;
 		FVCamConnection& VCamConnection = Connection.Value;
 
-		const bool bDidConnectSuccessfully = VCamConnection.AttemptConnection(VCam);
+		bool bDidConnectSuccessfully= false;
+		if (VCamConnection.ConnectionTargetSettings.HasValidSettings())
+		{
+			bDidConnectSuccessfully = VCamConnection.AttemptConnection(VCamComponent);
 
-		if (!bDidConnectSuccessfully)
+			if (!bDidConnectSuccessfully)
+			{
+				UE_LOG(LogVCamConnection, Warning, TEXT("Widget %s: Failed to create for VCam Connection with Connection Name: %s"), *GetName(), *ConnectionName.ToString());	
+			}
+		}
+		else
 		{
-			UE_LOG(LogVCamConnection, Warning, TEXT("Widget %s: Failed to create for VCam Connection with Connection Name: %s"), *GetName(), *ConnectionName.ToString());	
+			VCamConnection.ResetConnection();
 		}
 		
 		OnConnectionUpdated(ConnectionName, bDidConnectSuccessfully, VCamConnection.ConnectionTargetSettings.TargetConnectionPoint, VCamConnection.ConnectedModifier, VCamConnection.ConnectedAction);
 	}
 
-	// Register for input with the VCam Component if desired
-	if (bRegisterForInput)
+	return true;
+}
+
+void UVCamWidget::UpdateConnectionTargets(const TMap<FName, FVCamConnectionTargetSettings>& NewConnectionTargets, const bool bReinitializeOnSuccessfulUpdate, EConnectionUpdateResult& Result)
+{
+	bool bConnectionsUpdated = false;
+	for (const TPair<FName, FVCamConnectionTargetSettings>& NewConnectionTarget : NewConnectionTargets)
+	{
+		const FName& ConnectionName = NewConnectionTarget.Key;
+		const FVCamConnectionTargetSettings& TargetSettings = NewConnectionTarget.Value;
+
+		if (FVCamConnection* CurrentConnection = Connections.Find(ConnectionName))
+		{
+			if (CurrentConnection->ConnectionTargetSettings == TargetSettings)
+			{
+				continue;
+			}
+
+			CurrentConnection->ConnectionTargetSettings = TargetSettings;
+			bConnectionsUpdated = true;
+		}
+	}
+
+	if (bConnectionsUpdated)
+	{
+		Result = EConnectionUpdateResult::DidUpdateConnections;
+		if (bReinitializeOnSuccessfulUpdate)
+		{
+			ReinitializeConnections();
+		}
+	}
+	else
 	{
-		VCam->RegisterObjectForInput(this);
-		VCam->AddInputMappingContext(InputMappingContext, InputContextPriority);
+		Result = EConnectionUpdateResult::NoConnectionsUpdated;
 	}
 }
diff --git a/Engine/Plugins/Experimental/VirtualCamera/Source/VCamCore/Public/UI/VCamConnectionStructs.h b/Engine/Plugins/Experimental/VirtualCamera/Source/VCamCore/Public/UI/VCamConnectionStructs.h
--- a/Engine/Plugins/Experimental/VirtualCamera/Source/VCamCore/Public/UI/VCamConnectionStructs.h
+++ b/Engine/Plugins/Experimental/VirtualCamera/Source/VCamCore/Public/UI/VCamConnectionStructs.h
@@ -35,6 +35,17 @@ struct VCAMCORE_API FVCamConnectionTargetSettings
 	 */
 	UPROPERTY(EditAnywhere, BlueprintReadWrite, Category=ConnectionBinding)
 	FName TargetConnectionPoint = NAME_None;
+
+	bool operator==(const FVCamConnectionTargetSettings& Other) const
+	{
+		return TargetModifierName == Other.TargetModifierName
+		&& TargetConnectionPoint == Other.TargetConnectionPoint;
+	}
+
+	bool HasValidSettings() const
+	{
+		return !(TargetModifierName.IsNone() || TargetConnectionPoint.IsNone());
+	}
 };
 
 /*
@@ -113,4 +124,9 @@ struct VCAMCORE_API FVCamConnection
 	 * Attempts to make a connection to a target Connection Point and returns a bool indicating success
 	 */
 	bool AttemptConnection(UVCamComponent* VCamComponent);
+
+	/*
+	 * Clears any currently connected modifier and action
+	 */
+	void ResetConnection();
 };
diff --git a/Engine/Plugins/Experimental/VirtualCamera/Source/VCamCore/Public/UI/VCamWidget.h b/Engine/Plugins/Experimental/VirtualCamera/Source/VCamCore/Public/UI/VCamWidget.h
--- a/Engine/Plugins/Experimental/VirtualCamera/Source/VCamCore/Public/UI/VCamWidget.h
+++ b/Engine/Plugins/Experimental/VirtualCamera/Source/VCamCore/Public/UI/VCamWidget.h
@@ -12,6 +12,13 @@
 class UInputAction;
 class UVCamModifier;
 
+UENUM(BlueprintType)
+enum class EConnectionUpdateResult : uint8
+{
+	DidUpdateConnections,
+	NoConnectionsUpdated
+};
+
 /*
  * A wrapper widget class that contains a set of VCam Connections
  * 
@@ -63,7 +70,31 @@ public:
 
 	/*
 	 * Iterate all VCam Connections within the widget and attempt to connect them using the provided VCam Component
+	 * Optionally registers for input with the VCam Component if requested
 	 */
-	UFUNCTION(BlueprintCallable, Category="VCam Connections")
 	void InitializeConnections(UVCamComponent* VCam);
+
+	/*
+	 * Attempts to use the cached VCam Component to initialize the connections
+	 * Returns true if there was a valid VCam Component to initialize with
+	 */
+	UFUNCTION(BlueprintCallable, Category="VCam Connections")
+	bool ReinitializeConnections();
+
+	/*
+	 * Looks through the set of Connections on this widget and if a matching connection name is found will attempt to update the connection target.
+	 * If the new target is the same as the old target then no update is performed.
+	 * This function will optionally reinitialize the widget connections if 1 or more targets were updated.
+	 * 
+	 * Note: If a connection name in the NewConnectionTargets map is not found in the Widget Connections then it is ignored
+	 * no new connections will be created and no connections will be removed
+	 */
+	UFUNCTION(BlueprintCallable, Category="VCam Connections", meta=(ExpandEnumAsExecs="Result", bReinitializeOnSuccessfulUpdate=true))
+	void UpdateConnectionTargets(const TMap<FName, FVCamConnectionTargetSettings>& NewConnectionTargets, const bool bReinitializeOnSuccessfulUpdate, EConnectionUpdateResult& Result);
+
+	/*
+	 * Cached pointer to the VCam Component that owns this VCam Widget
+	 */
+	UPROPERTY(Transient, BlueprintReadOnly, Category="VCam")
+	TObjectPtr<UVCamComponent> VCamComponent;
 };
\ No newline at end of file
