Reverted UVCamOutputProviderBase::NotifyWidgetOfComponentChange to use ForEachWidget due to an issue with ForEachWidgetAndDescendants

diff --git a/Engine/Plugins/Experimental/VirtualCamera/Source/VCamCore/Private/VCamOutputProviderBase.cpp b/Engine/Plugins/Experimental/VirtualCamera/Source/VCamCore/Private/VCamOutputProviderBase.cpp
--- a/Engine/Plugins/Experimental/VirtualCamera/Source/VCamCore/Private/VCamOutputProviderBase.cpp
+++ b/Engine/Plugins/Experimental/VirtualCamera/Source/VCamCore/Private/VCamOutputProviderBase.cpp
@@ -229,7 +229,7 @@ void UVCamOutputProviderBase::NotifyWidgetOfComponentChange() const
 				// Find all VCam Widgets inside the displayed widget and Initialize them with the owning VCam Component
 				if (IsValid(DisplayedWidget->WidgetTree))
 				{
-					DisplayedWidget->WidgetTree->ForEachWidgetAndDescendants([VCamComponent](UWidget* Widget)
+					DisplayedWidget->WidgetTree->ForEachWidget([VCamComponent](UWidget* Widget)
 					{
 						if (UVCamWidget* VCamWidget = Cast<UVCamWidget>(Widget))
 						{
