Changed Right Click Asset Creation Menu to "VCam" from "Virtual Camera"

diff --git a/Engine/Plugins/Experimental/VirtualCamera/Source/VCamCoreEditor/Private/VCamModifierFactory.cpp b/Engine/Plugins/Experimental/VirtualCamera/Source/VCamCoreEditor/Private/VCamModifierFactory.cpp
--- a/Engine/Plugins/Experimental/VirtualCamera/Source/VCamCoreEditor/Private/VCamModifierFactory.cpp
+++ b/Engine/Plugins/Experimental/VirtualCamera/Source/VCamCoreEditor/Private/VCamModifierFactory.cpp
@@ -36,7 +36,7 @@ UObject* UVCamModifierFactory::FactoryCreateNew(UClass* Class, UObject* InParent
 uint32 UVCamModifierFactory::GetMenuCategories() const
 {
 	IAssetTools& AssetTools = FModuleManager::LoadModuleChecked<FAssetToolsModule>("AssetTools").Get();
-	return AssetTools.RegisterAdvancedAssetCategory("VirtualCamera", LOCTEXT("AssetCategoryName", "Virtual Camera"));
+	return AssetTools.RegisterAdvancedAssetCategory("VirtualCamera", LOCTEXT("AssetCategoryName", "VCam"));
 }
 
 #undef LOCTEXT_NAMESPACE
\ No newline at end of file
