Updated UVCamComponent::GetEnhancedInputSubsystemInterface to return the player input system in non-editor situations

diff --git a/Engine/Plugins/Experimental/VirtualCamera/Source/VCamCore/Private/VCamComponent.cpp b/Engine/Plugins/Experimental/VirtualCamera/Source/VCamCore/Private/VCamComponent.cpp
--- a/Engine/Plugins/Experimental/VirtualCamera/Source/VCamCore/Private/VCamComponent.cpp
+++ b/Engine/Plugins/Experimental/VirtualCamera/Source/VCamCore/Private/VCamComponent.cpp
@@ -11,6 +11,7 @@
 #include "GameDelegates.h"
 #include "Engine/GameEngine.h"
 #include "EnhancedInputComponent.h"
+#include "EnhancedInputSubsystems.h"
 #include "Engine/InputDelegateBinding.h"
 #include "InputMappingContext.h"
 
@@ -164,8 +165,16 @@ void UVCamComponent::NotifyComponentWasReplaced(UVCamComponent* ReplacementCompo
 
 IEnhancedInputSubsystemInterface* UVCamComponent::GetEnhancedInputSubsystemInterface() const
 {
+	if (UWorld* World = GetWorld(); World->IsGameWorld())
+	{
+		
+		if (ULocalPlayer* FirstLocalPlayer = World->GetFirstLocalPlayerFromController())
+		{
+			return ULocalPlayer::GetSubsystem<UEnhancedInputLocalPlayerSubsystem>(FirstLocalPlayer);
+		}
+	}
 #if WITH_EDITOR
-	if (GEditor)
+	else if (GEditor)
 	{
 		return GEditor->GetEditorSubsystem<UEnhancedInputEditorSubsystem>();
 	}
