Fix for uninitialized struct members in FVCamConnection

diff --git a/Engine/Plugins/Experimental/VirtualCamera/Source/VCamCore/Public/UI/VCamConnectionStructs.h b/Engine/Plugins/Experimental/VirtualCamera/Source/VCamCore/Public/UI/VCamConnectionStructs.h
--- a/Engine/Plugins/Experimental/VirtualCamera/Source/VCamCore/Public/UI/VCamConnectionStructs.h
+++ b/Engine/Plugins/Experimental/VirtualCamera/Source/VCamCore/Public/UI/VCamConnectionStructs.h
@@ -84,14 +84,14 @@ struct VCAMCORE_API FVCamConnection
 	 * If this connection requires an Input Action then this property specifies what Input Action Type the Connection Point must provide to be considered a valid connection
 	 */
 	UPROPERTY(EditAnywhere, BlueprintReadWrite, Category=Connection, meta=(EditCondition="bRequiresInputAction"))
-	EInputActionValueType ActionType;
+	EInputActionValueType ActionType = EInputActionValueType::Boolean;
 
 	/*
 	 * If you are not using an external system to configure the connection then you can enable this to have
 	 * explicit control over which Modifier and Connection Point to look for when attempting a connection
 	 */
 	UPROPERTY(EditAnywhere, BlueprintReadWrite, Category=Connection)
-	bool bManuallyConfigureConnection;
+	bool bManuallyConfigureConnection = false;
 
 	/*
 	 * A struct containing information about which Modifier and Connection Point to look for when attempting a connection
