Updated the Virtual Camera Actor placement settings to be aware of the new deprecated folder location

diff --git a/Engine/Plugins/Experimental/VirtualCamera/Source/VirtualCameraEditor/Private/VirtualCameraEditorModule.cpp b/Engine/Plugins/Experimental/VirtualCamera/Source/VirtualCameraEditor/Private/VirtualCameraEditorModule.cpp
--- a/Engine/Plugins/Experimental/VirtualCamera/Source/VirtualCameraEditor/Private/VirtualCameraEditorModule.cpp
+++ b/Engine/Plugins/Experimental/VirtualCamera/Source/VirtualCameraEditor/Private/VirtualCameraEditorModule.cpp
@@ -56,8 +56,8 @@ public:
 			if (const FPlacementCategoryInfo* Info = IVPUtilitiesEditorModule::Get().GetVirtualProductionPlacementCategoryInfo())
 			{
 				FAssetData SimpleVirtualCameraAssetData(
-					TEXT("/VirtualCamera/VCamCore/Blueprints/SimpleVirtualCamera"),
-					TEXT("/VirtualCamera/VCamCore/Blueprints"),
+					TEXT("/VirtualCamera/Deprecated/VCamCore/Blueprints/SimpleVirtualCamera"),
+					TEXT("/VirtualCamera/Deprecated/VCamCore/Blueprints"),
 					TEXT("SimpleVirtualCamera"),
 					FTopLevelAssetPath(TEXT("/Script/Engine"), TEXT("Blueprint"))
 				);
@@ -74,8 +74,8 @@ public:
 					));
 
 				FAssetData VirtualCamera2ActorAssetData(
-					TEXT("/VirtualCamera/V2/VirtualCamera2Actor"),
-					TEXT("/VirtualCamera/V2"),
+					TEXT("/VirtualCamera/Deprecated/V2/VirtualCamera2Actor"),
+					TEXT("/VirtualCamera/Deprecated/V2"),
 					TEXT("VirtualCamera2Actor"),
 					FTopLevelAssetPath(TEXT("/Script/Engine"), TEXT("Blueprint"))
 				);
