Fix for possible nullptr when dereferencing in UVCamComponent::GetEnhancedInputSubsystemInterface

diff --git a/Engine/Plugins/Experimental/VirtualCamera/Source/VCamCore/Private/VCamComponent.cpp b/Engine/Plugins/Experimental/VirtualCamera/Source/VCamCore/Private/VCamComponent.cpp
--- a/Engine/Plugins/Experimental/VirtualCamera/Source/VCamCore/Private/VCamComponent.cpp
+++ b/Engine/Plugins/Experimental/VirtualCamera/Source/VCamCore/Private/VCamComponent.cpp
@@ -165,10 +165,10 @@ void UVCamComponent::NotifyComponentWasReplaced(UVCamComponent* ReplacementCompo
 
 IEnhancedInputSubsystemInterface* UVCamComponent::GetEnhancedInputSubsystemInterface() const
 {
-	if (UWorld* World = GetWorld(); World->IsGameWorld())
+	if (const UWorld* World = GetWorld(); IsValid(World) && World->IsGameWorld())
 	{
 		
-		if (ULocalPlayer* FirstLocalPlayer = World->GetFirstLocalPlayerFromController())
+		if (const ULocalPlayer* FirstLocalPlayer = World->GetFirstLocalPlayerFromController())
 		{
 			return ULocalPlayer::GetSubsystem<UEnhancedInputLocalPlayerSubsystem>(FirstLocalPlayer);
 		}
