Added Input Contexts settings to the VCam Modifier class directly and changed the VCam Component to read from this new location.

Removal of the GetInputMappingContextAndPriority function will happen in a later CL to allow for Contexts to be moved over inside existing modifiers

diff --git a/Engine/Plugins/Experimental/VirtualCamera/Source/VCamCore/Private/VCamComponent.cpp b/Engine/Plugins/Experimental/VirtualCamera/Source/VCamCore/Private/VCamComponent.cpp
--- a/Engine/Plugins/Experimental/VirtualCamera/Source/VCamCore/Private/VCamComponent.cpp
+++ b/Engine/Plugins/Experimental/VirtualCamera/Source/VCamCore/Private/VCamComponent.cpp
@@ -410,8 +410,8 @@ void UVCamComponent::AddInputMappingContext(const UVCamModifier* Modifier)
 {
 	if (IEnhancedInputSubsystemInterface* EnhancedInputSubsystemInterface = GetEnhancedInputSubsystemInterface())
 	{
-		int32 InputPriority;
-		const UInputMappingContext* IMC = Modifier->GetInputMappingContext(InputPriority);
+		const int32 InputPriority = Modifier->InputContextPriority;
+		const UInputMappingContext* IMC = Modifier->InputMappingContext;
 		if (IsValid(IMC))
 		{
 			if (!EnhancedInputSubsystemInterface->HasMappingContext(IMC))
diff --git a/Engine/Plugins/Experimental/VirtualCamera/Source/VCamCore/Private/VCamModifier.cpp b/Engine/Plugins/Experimental/VirtualCamera/Source/VCamCore/Private/VCamModifier.cpp
--- a/Engine/Plugins/Experimental/VirtualCamera/Source/VCamCore/Private/VCamModifier.cpp
+++ b/Engine/Plugins/Experimental/VirtualCamera/Source/VCamCore/Private/VCamModifier.cpp
@@ -29,18 +29,6 @@ void UVCamBlueprintModifier::Apply(UVCamModifierContext* Context, UCineCameraCom
 	}
 }
 
-const UInputMappingContext* UVCamBlueprintModifier::GetInputMappingContext(int32& InputPriority) const
-{
-	// Forward the call to the Blueprint implementation
-	{
-		FEditorScriptExecutionGuard ScriptGuard;
-
-		UInputMappingContext* MappingContext = nullptr;
-		GetInputMappingContextAndPriority(MappingContext, InputPriority);
-		return MappingContext;
-	}
-}
-
 void UVCamModifier::Initialize(UVCamModifierContext* Context, UInputComponent* InputComponent /* = nullptr */)
 {
 	// Binds any dynamic input delegates to the provided input component
@@ -120,11 +108,6 @@ FName UVCamModifier::GetStackEntryName() const
 	return NAME_None;
 }
 
-const UInputMappingContext* UVCamModifier::GetInputMappingContext(int32& InputPriority) const
-{
-	return nullptr;
-}
-
 FModifierStackEntry* UVCamModifier::GetCorrespondingStackEntry() const
 {
 	FModifierStackEntry* StackEntry = nullptr;
diff --git a/Engine/Plugins/Experimental/VirtualCamera/Source/VCamCore/Public/VCamModifier.h b/Engine/Plugins/Experimental/VirtualCamera/Source/VCamCore/Public/VCamModifier.h
--- a/Engine/Plugins/Experimental/VirtualCamera/Source/VCamCore/Public/VCamModifier.h
+++ b/Engine/Plugins/Experimental/VirtualCamera/Source/VCamCore/Public/VCamModifier.h
@@ -61,9 +61,13 @@ public:
 	UFUNCTION(BlueprintPure, Category = "VirtualCamera", meta=(ReturnDisplayName="Name"))
 	FName GetStackEntryName() const;
 
-	// Allows a modifier to return Input Mapping Context which will get automatically registered with the input system
-	// The Input Priority of the mapping context will be set by reference
-	virtual const UInputMappingContext* GetInputMappingContext(int32& InputPriority) const;
+	// If an Input Mapping Context is specified then that Context will be automatically added to the input system when this Modifier is Initialized 
+	UPROPERTY(EditAnywhere, BlueprintReadOnly, Category="VCam Input", meta=(EditCondition="bRegisterForInput"))
+	TObjectPtr<UInputMappingContext> InputMappingContext = nullptr;
+
+	// If an Input Mapping Context is provided then this value defines the priority level that the context is added to the input system with
+	UPROPERTY(EditAnywhere, BlueprintReadOnly, Category="VCam Input", meta=(EditCondition="bRegisterForInput"))
+	int32 InputContextPriority = 0;
 
 	UPROPERTY(EditAnywhere, BlueprintReadWrite, Category = "VCam Connection Points")
 	TMap<FName, FVCamModifierConnectionPoint> ConnectionPoints;
@@ -84,7 +88,6 @@ class VCAMCORE_API UVCamBlueprintModifier : public UVCamModifier
 public:
 	virtual void Initialize(UVCamModifierContext* Context, UInputComponent* InputComponent=nullptr) override;
 	virtual void Apply(UVCamModifierContext* Context, UCineCameraComponent* CameraComponent, const float DeltaTime) override;
-	virtual const UInputMappingContext* GetInputMappingContext(int32& InputPriority) const override;
 	
 	UFUNCTION(BlueprintImplementableEvent, Category="VirtualCamera")
 	void OnInitialize(UVCamModifierContext* Context);
@@ -92,7 +95,8 @@ public:
 	UFUNCTION(BlueprintImplementableEvent, Category="VirtualCamera")
 	void OnApply(UVCamModifierContext* Context, UCineCameraComponent* CameraComponent, const float DeltaTime);
 
-	// Allows a modifier to return Input Mapping Context which will get automatically registered with the input system
+	// This function is deliberately non-working to force cleanup of Input Contexts
+	// Please move any previous values to the new properties in Class Defaults
 	UFUNCTION(BlueprintImplementableEvent, Category="VirtualCamera")
-	void GetInputMappingContextAndPriority(UInputMappingContext*& InputMappingContext, int32& InputPriority) const;
+	void GetInputMappingContextAndPriority() const;
 };
\ No newline at end of file
