- Implemented an intermediate dialogue when creating VCam related blueprints to prompt for optional InputActionContext
- Updated the VCam Modifier Factory to use the new Dialog
- Added a VCam Widget Category that also uses the new Dialog

diff --git a/Engine/Plugins/Experimental/VirtualCamera/Source/VCamCoreEditor/Private/VCamModifierFactory.cpp b/Engine/Plugins/Experimental/VirtualCamera/Source/VCamCoreEditor/Private/VCamModifierFactory.cpp
--- a/Engine/Plugins/Experimental/VirtualCamera/Source/VCamCoreEditor/Private/VCamModifierFactory.cpp
+++ b/Engine/Plugins/Experimental/VirtualCamera/Source/VCamCoreEditor/Private/VCamModifierFactory.cpp
@@ -4,6 +4,7 @@
 #include "VCamModifier.h"
 
 #include "AssetToolsModule.h"
+#include "Kismet2/BlueprintEditorUtils.h"
 #include "Kismet2/KismetEditorUtilities.h"
 
 #define LOCTEXT_NAMESPACE "VCamModifierFactory"
@@ -22,6 +23,11 @@ FText UVCamModifierFactory::GetDisplayName() const
 	return LOCTEXT("DisplayName", "VCam Modifier");
 }
 
+FText UVCamModifierFactory::GetToolTip() const
+{
+	return LOCTEXT("Tooltip", "Modifiers can manipulate the camera associated with a VCam Component. You can create custom modifiers with Blueprints or C++ and add them to the modifier stack of a VCam Component to layer different effects.");
+}
+
 UObject* UVCamModifierFactory::FactoryCreateNew(UClass* Class, UObject* InParent, FName Name, EObjectFlags Flags, UObject* Context, FFeedbackContext* Warn)
 {
 	UBlueprint* ModifierBlueprint = nullptr;
@@ -29,6 +35,14 @@ UObject* UVCamModifierFactory::FactoryCreateNew(UClass* Class, UObject* InParent
 	{
 		ensure(0 != (RF_Public & Flags));
 		ModifierBlueprint = FKismetEditorUtilities::CreateBlueprint(ParentClass, InParent, Name, BPTYPE_Normal, UBlueprint::StaticClass(), UBlueprintGeneratedClass::StaticClass(), NAME_None);
+		if (TSubclassOf<UObject> GeneratedClass = ModifierBlueprint->GeneratedClass)
+		{
+			if (UVCamBlueprintModifier* DefaultSubject = GeneratedClass->GetDefaultObject<UVCamBlueprintModifier>())
+			{
+				DefaultSubject->InputMappingContext = InputMappingContext;
+			}
+		}
+		FBlueprintEditorUtils::MarkBlueprintAsModified(ModifierBlueprint);
 	}
 	return ModifierBlueprint;	
 }
@@ -39,4 +53,4 @@ uint32 UVCamModifierFactory::GetMenuCategories() const
 	return AssetTools.RegisterAdvancedAssetCategory("VirtualCamera", LOCTEXT("AssetCategoryName", "VCam"));
 }
 
-#undef LOCTEXT_NAMESPACE
\ No newline at end of file
+#undef LOCTEXT_NAMESPACE
diff --git a/Engine/Plugins/Experimental/VirtualCamera/Source/VCamCoreEditor/Private/VCamObjectWithInputFactory.cpp b/Engine/Plugins/Experimental/VirtualCamera/Source/VCamCoreEditor/Private/VCamObjectWithInputFactory.cpp
new file mode 100644
--- /dev/null
+++ b/Engine/Plugins/Experimental/VirtualCamera/Source/VCamCoreEditor/Private/VCamObjectWithInputFactory.cpp
@@ -0,0 +1,124 @@
+// Copyright Epic Games, Inc. All Rights Reserved.
+
+#include "VCamObjectWithInputFactory.h"
+
+#include "AssetThumbnail.h"
+#include "Editor.h"
+#include "PropertyCustomizationHelpers.h"
+
+#include "Widgets/SBoxPanel.h"
+#include "Widgets/SWindow.h"
+#include "Widgets/Input/SButton.h"
+#include "Widgets/Layout/SBorder.h"
+
+#define LOCTEXT_NAMESPACE "VCamObjectFactories"
+
+bool UVCamObjectWithInputFactory::ConfigureProperties()
+{
+	class FVCamWidgetFactoryUI : public TSharedFromThis<FVCamWidgetFactoryUI>
+	{
+	public:
+		FReply OnCreate()
+		{
+			if (PickerWindow.IsValid())
+			{
+				PickerWindow->RequestDestroyWindow();
+			}
+			bWasCanceled= false;
+			return FReply::Handled();
+		}
+
+		FReply OnCancel()
+		{
+			if (PickerWindow.IsValid())
+			{
+				PickerWindow->RequestDestroyWindow();
+			}
+			bWasCanceled = true;
+			return FReply::Handled();
+		}
+
+		bool WasCanceled() const
+		{
+			return bWasCanceled;
+		}
+
+		void OpenMappingSelector(UVCamObjectWithInputFactory* Factory)
+		{
+			AssetThumbnailPool = MakeShared<FAssetThumbnailPool>(1);
+			PickerWindow = SNew(SWindow)
+				.Title(LOCTEXT("PickerTitle", "Select Input Mapping Context"))
+				.ClientSize(FVector2D(350, 100))
+				.SupportsMinimize(false)
+				.SupportsMaximize(false)
+				[
+					SNew(SBorder)
+					.BorderImage(FAppStyle::GetBrush("Menu.Background"))
+					.Padding(10)
+					[
+						SNew(SVerticalBox)
+						+ SVerticalBox::Slot()
+						.FillHeight(1.f)
+						[
+							SNew(SObjectPropertyEntryBox)
+							.DisplayThumbnail(true)
+							.ThumbnailPool(AssetThumbnailPool)
+							.AllowClear(true)
+							.DisplayUseSelected(false)
+							.DisplayBrowse(false)
+							.AllowedClass(UInputMappingContext::StaticClass())
+							.ObjectPath_Lambda([Factory]() -> FString
+							{
+								return (Factory && Factory->InputMappingContext) ? Factory->InputMappingContext->GetPathName() : TEXT("None");
+							})
+							.OnObjectChanged_Lambda([Factory](const FAssetData& AssetData) -> void
+							{
+								if (Factory)
+								{
+									Factory->InputMappingContext = Cast<UInputMappingContext>(AssetData.GetAsset());
+								}
+							})
+						]
+						+ SVerticalBox::Slot()
+						.HAlign(HAlign_Right)
+						.VAlign(VAlign_Bottom)
+						.AutoHeight()
+						[
+							SNew(SHorizontalBox)
+							+ SHorizontalBox::Slot()
+							.AutoWidth()
+							[
+								SNew(SButton)
+								.Text(LOCTEXT("OK", "OK"))
+								.OnClicked(this, &FVCamWidgetFactoryUI::OnCreate)
+							]
+							+ SHorizontalBox::Slot()
+							.AutoWidth()
+							[
+								SNew(SButton)
+								.Text(LOCTEXT("Cancel", "Cancel"))
+								.OnClicked(this, &FVCamWidgetFactoryUI::OnCancel)
+							]
+						]
+					]
+				];
+			if (GEditor)
+			{
+				GEditor->EditorAddModalWindow(PickerWindow.ToSharedRef());
+			}
+			PickerWindow.Reset();
+		}
+
+	private:
+		TSharedPtr<SWindow> PickerWindow;
+		TSharedPtr<FAssetThumbnailPool> AssetThumbnailPool;
+		bool bWasCanceled = false;
+	};
+
+	const TSharedRef<FVCamWidgetFactoryUI> InputMappingSelector = MakeShared<FVCamWidgetFactoryUI>();
+	InputMappingSelector->OpenMappingSelector(this);
+	
+	return !InputMappingSelector->WasCanceled();
+}
+
+#undef LOCTEXT_NAMESPACE
\ No newline at end of file
diff --git a/Engine/Plugins/Experimental/VirtualCamera/Source/VCamCoreEditor/Private/VCamWidgetFactory.cpp b/Engine/Plugins/Experimental/VirtualCamera/Source/VCamCoreEditor/Private/VCamWidgetFactory.cpp
new file mode 100644
--- /dev/null
+++ b/Engine/Plugins/Experimental/VirtualCamera/Source/VCamCoreEditor/Private/VCamWidgetFactory.cpp
@@ -0,0 +1,57 @@
+// Copyright Epic Games, Inc. All Rights Reserved.
+
+#include "VCamWidgetFactory.h"
+#include "UI/VCamWidget.h"
+
+#include "AssetToolsModule.h"
+#include "Kismet2/BlueprintEditorUtils.h"
+#include "Kismet2/KismetEditorUtilities.h"
+#include "WidgetBlueprint.h"
+
+#define LOCTEXT_NAMESPACE "VCamWidgetFactory"
+
+UVCamWidgetFactory::UVCamWidgetFactory()
+{
+	SupportedClass = UBlueprint::StaticClass();
+	ParentClass = UVCamWidget::StaticClass();
+
+	bCreateNew = true;
+	bEditAfterNew = true;
+}
+
+FText UVCamWidgetFactory::GetDisplayName() const
+{
+	return LOCTEXT("DisplayName", "VCam Widget");
+}
+
+FText UVCamWidgetFactory::GetToolTip() const
+{
+	return LOCTEXT("Tooltip", "A wrapper widget class that contains a set of VCam Connections\n\nIf you add a widget deriving from VCam Widget to an Overlay Widget for a VCam Output Provider then when the Overlay is created by the Provider it will also call InitializeConnections with the owning VCam Component.");
+}
+
+UObject* UVCamWidgetFactory::FactoryCreateNew(UClass* Class, UObject* InParent, FName Name, EObjectFlags Flags, UObject* Context, FFeedbackContext* Warn)
+{
+	UBlueprint* WidgetBlueprint = nullptr;
+	if (ensure(SupportedClass == Class))
+	{
+		ensure(0 != (RF_Public & Flags));
+		WidgetBlueprint = FKismetEditorUtilities::CreateBlueprint(ParentClass, InParent, Name, BPTYPE_Normal, UWidgetBlueprint::StaticClass(), UWidgetBlueprintGeneratedClass::StaticClass(), NAME_None);
+		if (TSubclassOf<UObject> GeneratedClass = WidgetBlueprint->GeneratedClass)
+		{
+			if (UVCamWidget* DefaultSubject = GeneratedClass->GetDefaultObject<UVCamWidget>())
+			{
+				DefaultSubject->InputMappingContext = InputMappingContext;
+			}
+		}
+		FBlueprintEditorUtils::MarkBlueprintAsModified(WidgetBlueprint);
+	}
+	return WidgetBlueprint;	
+}
+
+uint32 UVCamWidgetFactory::GetMenuCategories() const
+{
+	IAssetTools& AssetTools = FModuleManager::LoadModuleChecked<FAssetToolsModule>("AssetTools").Get();
+	return AssetTools.RegisterAdvancedAssetCategory("VirtualCamera", LOCTEXT("AssetCategoryName", "VCam"));
+}
+
+#undef LOCTEXT_NAMESPACE
diff --git a/Engine/Plugins/Experimental/VirtualCamera/Source/VCamCoreEditor/Public/VCamModifierFactory.h b/Engine/Plugins/Experimental/VirtualCamera/Source/VCamCoreEditor/Public/VCamModifierFactory.h
--- a/Engine/Plugins/Experimental/VirtualCamera/Source/VCamCoreEditor/Public/VCamModifierFactory.h
+++ b/Engine/Plugins/Experimental/VirtualCamera/Source/VCamCoreEditor/Public/VCamModifierFactory.h
@@ -3,23 +3,20 @@
 #pragma once
 
 #include "CoreMinimal.h"
-#include "Factories/Factory.h"
+#include "VCamObjectWithInputFactory.h"
 
 #include "VCamModifierFactory.generated.h"
 
-UCLASS(hidecategories = Object)
-class UVCamModifierFactory : public UFactory
+UCLASS()
+class VCAMCOREEDITOR_API UVCamModifierFactory : public UVCamObjectWithInputFactory
 {
 	GENERATED_BODY()
 public:
 	UVCamModifierFactory();
 	
-	// The parent class of the created blueprint
-    UPROPERTY(EditAnywhere, Category="BlueprintVirtualSubjectFactory", meta=(AllowAbstract = "", BlueprintBaseOnly = ""))
-    TSubclassOf<UObject> ParentClass;
-	
 	//~ Begin UFactory Interface
 	virtual FText GetDisplayName() const override;
+	virtual FText GetToolTip() const override;
 	virtual UObject* FactoryCreateNew(UClass* Class, UObject* InParent, FName Name, EObjectFlags Flags, UObject* Context, FFeedbackContext* Warn) override;
 	virtual uint32 GetMenuCategories() const override;
 	//~ Begin UFactory Interface
diff --git a/Engine/Plugins/Experimental/VirtualCamera/Source/VCamCoreEditor/Public/VCamObjectWithInputFactory.h b/Engine/Plugins/Experimental/VirtualCamera/Source/VCamCoreEditor/Public/VCamObjectWithInputFactory.h
new file mode 100644
--- /dev/null
+++ b/Engine/Plugins/Experimental/VirtualCamera/Source/VCamCoreEditor/Public/VCamObjectWithInputFactory.h
@@ -0,0 +1,29 @@
+// Copyright Epic Games, Inc. All Rights Reserved.
+
+#pragma once
+
+#include "CoreMinimal.h"
+#include "Factories/Factory.h"
+#include "InputMappingContext.h"
+
+#include "VCamObjectWithInputFactory.generated.h"
+
+UCLASS(hidecategories = Object)
+class VCAMCOREEDITOR_API UVCamObjectWithInputFactory : public UFactory
+{
+	GENERATED_BODY()
+public:	
+	// The parent class of the created blueprint
+	UPROPERTY(EditAnywhere, Category="BlueprintVirtualSubjectFactory", meta=(AllowAbstract = "", BlueprintBaseOnly = ""))
+	TSubclassOf<UObject> ParentClass;
+
+	UPROPERTY(BlueprintReadWrite, Category = "VCam Modifier Factory")
+	TObjectPtr<UInputMappingContext> InputMappingContext;
+
+	UPROPERTY(BlueprintReadWrite, Category = "VCam Modifier Factory")
+	int32 MappingPriority;
+
+	// Begin UFactory Interface
+	virtual bool ConfigureProperties() override;
+	// End UFactory Interface
+};
\ No newline at end of file
diff --git a/Engine/Plugins/Experimental/VirtualCamera/Source/VCamCoreEditor/Public/VCamWidgetFactory.h b/Engine/Plugins/Experimental/VirtualCamera/Source/VCamCoreEditor/Public/VCamWidgetFactory.h
new file mode 100644
--- /dev/null
+++ b/Engine/Plugins/Experimental/VirtualCamera/Source/VCamCoreEditor/Public/VCamWidgetFactory.h
@@ -0,0 +1,24 @@
+// Copyright Epic Games, Inc. All Rights Reserved.
+
+#pragma once
+
+#include "CoreMinimal.h"
+#include "VCamObjectWithInputFactory.h"
+
+#include "VCamWidgetFactory.generated.h"
+
+UCLASS()
+class VCAMCOREEDITOR_API UVCamWidgetFactory : public UVCamObjectWithInputFactory
+{
+	GENERATED_BODY()
+public:
+	UVCamWidgetFactory();
+	
+	//~ Begin UFactory Interface
+	virtual FText GetDisplayName() const override;
+	virtual FText GetToolTip() const override;
+	virtual UObject* FactoryCreateNew(UClass* Class, UObject* InParent, FName Name, EObjectFlags Flags, UObject* Context, FFeedbackContext* Warn) override;
+	virtual uint32 GetMenuCategories() const override;
+	//~ Begin UFactory Interface
+	
+};
\ No newline at end of file
diff --git a/Engine/Plugins/Experimental/VirtualCamera/Source/VCamCoreEditor/VCamCoreEditor.build.cs b/Engine/Plugins/Experimental/VirtualCamera/Source/VCamCoreEditor/VCamCoreEditor.build.cs
--- a/Engine/Plugins/Experimental/VirtualCamera/Source/VCamCoreEditor/VCamCoreEditor.build.cs
+++ b/Engine/Plugins/Experimental/VirtualCamera/Source/VCamCoreEditor/VCamCoreEditor.build.cs
@@ -27,7 +27,11 @@ public class VCamCoreEditor : ModuleRules
 				"SlateCore",
 				"CinematicCamera",
 				"LiveLinkInterface",
-				"VCamCore"
+				"VCamCore",
+				"EnhancedInput",
+				"PropertyEditor",
+				"UMG", 
+				"UMGEditor"
 			}
 		);
 	}
