Added support for storing and applying input key overrides for the VCam system

diff --git a/Engine/Plugins/Experimental/VirtualCamera/Source/VCamCore/Private/UI/VCamWidget.cpp b/Engine/Plugins/Experimental/VirtualCamera/Source/VCamCore/Private/UI/VCamWidget.cpp
--- a/Engine/Plugins/Experimental/VirtualCamera/Source/VCamCore/Private/UI/VCamWidget.cpp
+++ b/Engine/Plugins/Experimental/VirtualCamera/Source/VCamCore/Private/UI/VCamWidget.cpp
@@ -2,18 +2,28 @@
 
 #include "UI/VCamWidget.h"
 
+void UVCamWidget::BeginDestroy()
+{
+	Super::BeginDestroy();
+	
+	if (VCamComponent)
+	{
+		VCamComponent->UnregisterObjectForInput(this);
+	}
+}
+
 void UVCamWidget::InitializeConnections(UVCamComponent* VCam)
 {
 	VCamComponent = VCam;
 
-	const bool bDidSuccessfullyInitialize = ReinitializeConnections();
-	
 	// Register for input with the VCam Component if desired
-	if (bDidSuccessfullyInitialize && bRegisterForInput)
+	if (VCamComponent && bRegisterForInput)
 	{
 		VCamComponent->RegisterObjectForInput(this);
 		VCamComponent->AddInputMappingContext(InputMappingContext, InputContextPriority);
 	}
+
+	ReinitializeConnections();
 }
 
 bool UVCamWidget::ReinitializeConnections()
@@ -47,6 +57,7 @@ bool UVCamWidget::ReinitializeConnections()
 		OnConnectionUpdated(ConnectionName, bDidConnectSuccessfully, VCamConnection.ConnectionTargetSettings.TargetConnectionPoint, VCamConnection.ConnectedModifier, VCamConnection.ConnectedAction);
 	}
 
+	PostConnectionsInitialized();
 	return true;
 }
 
diff --git a/Engine/Plugins/Experimental/VirtualCamera/Source/VCamCore/Private/VCamComponent.cpp b/Engine/Plugins/Experimental/VirtualCamera/Source/VCamCore/Private/VCamComponent.cpp
--- a/Engine/Plugins/Experimental/VirtualCamera/Source/VCamCore/Private/VCamComponent.cpp
+++ b/Engine/Plugins/Experimental/VirtualCamera/Source/VCamCore/Private/VCamComponent.cpp
@@ -90,6 +90,12 @@ UVCamComponent::UVCamComponent()
 			}
 		}
 #endif
+
+		// Apply the Default Input profile if possible
+		if (const UVCamInputSettings* VCamInputSettings = GetDefault<UVCamInputSettings>())
+		{
+			SetInputProfileFromName(VCamInputSettings->DefaultInputProfile);
+		}
 		
 	}
 }
@@ -237,6 +243,14 @@ void UVCamComponent::OnAttachmentChanged()
 #endif
 }
 
+void UVCamComponent::PostLoad()
+{
+	Super::PostLoad();
+
+	// Ensure the input profile is applied when this component is loaded
+	ApplyInputProfile();
+}
+
 #if WITH_EDITOR
 
 void UVCamComponent::CheckForErrors()
@@ -301,6 +315,7 @@ void UVCamComponent::PostEditChangeProperty(FPropertyChangedEvent& PropertyChang
 		static FName NAME_Enabled = GET_MEMBER_NAME_CHECKED(UVCamComponent, bEnabled);
 		static FName NAME_ModifierStack = GET_MEMBER_NAME_CHECKED(UVCamComponent, ModifierStack);
 		static FName NAME_TargetViewport = GET_MEMBER_NAME_CHECKED(UVCamComponent, TargetViewport);
+		static FName NAME_InputProfile = GET_MEMBER_NAME_CHECKED(UVCamComponent, InputProfile);
 
 		const FName PropertyName = Property->GetFName();
 
@@ -334,6 +349,10 @@ void UVCamComponent::PostEditChangeProperty(FPropertyChangedEvent& PropertyChang
 				}
 			}
 		}
+		else if (PropertyName == NAME_InputProfile)
+		{
+			ApplyInputProfile();
+		}
 
 		for (const UVCamOutputProviderBase* OutputProvider : OutputProviders)
 		{
@@ -433,6 +452,104 @@ void UVCamComponent::AddInputMappingContext(UInputMappingContext* Context, int32
 	}
 }
 
+bool UVCamComponent::SetInputProfileFromName(const FName ProfileName)
+{
+	if (const UVCamInputSettings* VCamInputSettings = GetDefault<UVCamInputSettings>())
+	{
+		if (const FVCamInputProfile* NewInputProfile = VCamInputSettings->InputProfiles.Find(ProfileName))
+		{
+			InputProfile = *NewInputProfile;
+			ApplyInputProfile();
+		}
+	}
+	return false;
+}
+
+bool UVCamComponent::AddInputProfileWithCurrentlyActiveMappings(const FName ProfileName, bool bUpdateIfProfileAlreadyExists)
+{
+	UVCamInputSettings* VCamInputSettings = GetMutableDefault<UVCamInputSettings>();
+
+	// If we don't have a valid settings object then early out
+	if (!VCamInputSettings)
+	{
+		return false;
+	}
+	
+	FVCamInputProfile* TargetInputProfile = VCamInputSettings->InputProfiles.Find(ProfileName);
+
+	// An Input Profile with this name already exists
+	if (TargetInputProfile)
+	{
+		if (!bUpdateIfProfileAlreadyExists)
+		{
+			return false;
+		}
+	}
+	else
+	{
+		TargetInputProfile = &VCamInputSettings->InputProfiles.Add(ProfileName);
+	}
+
+	if (const IEnhancedInputSubsystemInterface* EnhancedInputSubsystemInterface = GetEnhancedInputSubsystemInterface())
+	{
+		const TArray<FEnhancedActionKeyMapping> PlayerMappableActionKeyMappings = EnhancedInputSubsystemInterface->GetAllPlayerMappableActionKeyMappings();
+		for (const FEnhancedActionKeyMapping& PlayerMappableActionKeyMapping : PlayerMappableActionKeyMappings)
+		{
+			const FName MappingName = PlayerMappableActionKeyMapping.PlayerMappableOptions.Name;
+
+			// Prefer to use the current mapped key but fallback to the default if no key is mapped
+			FKey CurrentKey = EnhancedInputSubsystemInterface->GetPlayerMappedKey(MappingName);
+			if (!CurrentKey.IsValid())
+			{
+				CurrentKey = PlayerMappableActionKeyMapping.Key;
+			}
+
+			if (!TargetInputProfile->MappableKeyOverrides.Contains(MappingName))
+			{
+				TargetInputProfile->MappableKeyOverrides.Add(MappingName, CurrentKey);
+			}
+		}
+	}
+	
+	VCamInputSettings->SaveConfig();
+	return true;
+}
+
+bool UVCamComponent::SaveCurrentInputProfileToSettings(const FName ProfileName) const
+{
+	UVCamInputSettings* VCamInputSettings = GetMutableDefault<UVCamInputSettings>();
+
+	// If we don't have a valid settings object then early out
+	if (!VCamInputSettings)
+	{
+		return false;
+	}
+
+	FVCamInputProfile& SettingsInputProfile = VCamInputSettings->InputProfiles.FindOrAdd(ProfileName);
+	SettingsInputProfile = InputProfile;
+	VCamInputSettings->SaveConfig();
+
+	return true;
+}
+
+TArray<FEnhancedActionKeyMapping> UVCamComponent::GetAllPlayerMappableActionKeyMappings() const
+{
+	if (const IEnhancedInputSubsystemInterface* EnhancedInputSubsystemInterface = GetEnhancedInputSubsystemInterface())
+	{
+		return EnhancedInputSubsystemInterface->GetAllPlayerMappableActionKeyMappings();
+	}
+	return TArray<FEnhancedActionKeyMapping>();
+}
+
+FKey UVCamComponent::GetPlayerMappedKey(const FName MappingName) const
+{
+	if (const IEnhancedInputSubsystemInterface* EnhancedInputSubsystemInterface = GetEnhancedInputSubsystemInterface())
+	{
+		return EnhancedInputSubsystemInterface->GetPlayerMappedKey(MappingName);
+	}
+	return EKeys::Invalid;
+}
+
 
 void UVCamComponent::Update()
 {
@@ -844,6 +961,12 @@ void UVCamComponent::GetOutputProvidersByClass(TSubclassOf<UVCamOutputProviderBa
 	}
 }
 
+void UVCamComponent::SetInputProfile(const FVCamInputProfile& NewInputProfile)
+{
+	InputProfile = NewInputProfile;
+	ApplyInputProfile();
+}
+
 void UVCamComponent::GetLiveLinkDataForCurrentFrame(FLiveLinkCameraBlueprintData& LiveLinkData)
 {
 	IModularFeatures& ModularFeatures = IModularFeatures::Get();
@@ -887,24 +1010,21 @@ void UVCamComponent::RegisterObjectForInput(UObject* Object)
 	}
 }
 
-TArray<FEnhancedActionKeyMapping> UVCamComponent::GetPlayerMappableKeys() const
+void UVCamComponent::UnregisterObjectForInput(UObject* Object) const
 {
-	if (const IEnhancedInputSubsystemInterface* EnhancedInputSubsystemInterface = GetEnhancedInputSubsystemInterface())
+	if (IsValid(InputComponent) && Object)
 	{
-		return EnhancedInputSubsystemInterface->GetAllPlayerMappableActionKeyMappings();
+		InputComponent->ClearBindingsForObject(Object);
 	}
-	return {};
 }
 
-int32 UVCamComponent::AddPlayerMappedKey(const FName MappingName, const FKey NewKey,
-	const FModifyContextOptions& Options) const
+TArray<FEnhancedActionKeyMapping> UVCamComponent::GetPlayerMappableKeys() const
 {
-	if (IEnhancedInputSubsystemInterface* EnhancedInputSubsystemInterface = GetEnhancedInputSubsystemInterface())
+	if (const IEnhancedInputSubsystemInterface* EnhancedInputSubsystemInterface = GetEnhancedInputSubsystemInterface())
 	{
-		return EnhancedInputSubsystemInterface->AddPlayerMappedKey(MappingName, NewKey, Options);
+		return EnhancedInputSubsystemInterface->GetAllPlayerMappableActionKeyMappings();
 	}
-
-	return 0;
+	return {};
 }
 
 void UVCamComponent::InjectInputForAction(const UInputAction* Action, FInputActionValue RawValue, const TArray<UInputModifier*>& Modifiers, const TArray<UInputTrigger*>& Triggers)
@@ -923,6 +1043,26 @@ void UVCamComponent::InjectInputVectorForAction(const UInputAction* Action, FVec
 	}
 }
 
+void UVCamComponent::ApplyInputProfile()
+{
+	IEnhancedInputSubsystemInterface* EnhancedInputSubsystemInterface = GetEnhancedInputSubsystemInterface();
+	if (EnhancedInputSubsystemInterface)
+	{
+		EnhancedInputSubsystemInterface->RemoveAllPlayerMappedKeys();
+		for (const TPair<FName, FKey>& MappableKeyOverride : InputProfile.MappableKeyOverrides)
+		{
+			const FName& MappingName = MappableKeyOverride.Key;
+			const FKey& NewKey = MappableKeyOverride.Value;
+			
+			// Ensure we have a valid name to map
+			if (MappingName != NAME_None)
+			{
+				EnhancedInputSubsystemInterface->AddPlayerMappedKey(MappingName, NewKey);
+			}
+		}
+	}
+}
+
 void UVCamComponent::CopyLiveLinkDataToCamera(const FLiveLinkCameraBlueprintData& LiveLinkData, UCineCameraComponent* CameraComponent)
 {
 	const FLiveLinkCameraStaticData& StaticData = LiveLinkData.StaticData;
diff --git a/Engine/Plugins/Experimental/VirtualCamera/Source/VCamCore/Public/UI/VCamWidget.h b/Engine/Plugins/Experimental/VirtualCamera/Source/VCamCore/Public/UI/VCamWidget.h
--- a/Engine/Plugins/Experimental/VirtualCamera/Source/VCamCore/Public/UI/VCamWidget.h
+++ b/Engine/Plugins/Experimental/VirtualCamera/Source/VCamCore/Public/UI/VCamWidget.h
@@ -30,6 +30,8 @@ class VCAMCORE_API UVCamWidget : public UUserWidget
 {
 	GENERATED_BODY()
 public:
+	virtual void BeginDestroy() override;
+	
 	/*
 	 * The VCam Connections associated with this Widget
 	 * 
@@ -68,6 +70,12 @@ public:
 	UFUNCTION(BlueprintImplementableEvent, Category="VCam Connections")
 	void OnConnectionUpdated(FName ConnectionName, bool bDidConnectSuccessfully, FName ModifierConnectionPointName, UVCamModifier* ConnectedModifier, UInputAction* ConnectedAction);
 
+	/*
+	 * Event called after all connections have been initialized or reinitialized
+	 */
+	UFUNCTION(BlueprintImplementableEvent, Category="VCam Connections")
+	void PostConnectionsInitialized();
+
 	/*
 	 * Iterate all VCam Connections within the widget and attempt to connect them using the provided VCam Component
 	 * Optionally registers for input with the VCam Component if requested
diff --git a/Engine/Plugins/Experimental/VirtualCamera/Source/VCamCore/Public/VCamComponent.h b/Engine/Plugins/Experimental/VirtualCamera/Source/VCamCore/Public/VCamComponent.h
--- a/Engine/Plugins/Experimental/VirtualCamera/Source/VCamCore/Public/VCamComponent.h
+++ b/Engine/Plugins/Experimental/VirtualCamera/Source/VCamCore/Public/VCamComponent.h
@@ -6,6 +6,7 @@
 #include "EnhancedActionKeyMapping.h"
 #include "EnhancedInputSubsystemInterface.h"
 #include "VCamTypes.h"
+#include "VCamInputSettings.h"
 #include "Roles/LiveLinkCameraTypes.h"
 #include "VCamOutputProviderBase.h"
 #include "GameplayTagContainer.h"
@@ -55,6 +56,8 @@ public:
 
 	virtual void OnAttachmentChanged() override;
 
+	virtual void PostLoad() override;
+
 	TSharedPtr<FSceneViewport> GetTargetSceneViewport() const;
 	TWeakPtr<SWindow> GetTargetInputWindow() const;
 
@@ -89,6 +92,32 @@ public:
 	// Adds and explicitly provided Input Mapping Context to the input system
 	void AddInputMappingContext(UInputMappingContext* Context, int32 Priority);
 
+	// Attempts to apply key mapping settings from an input profile defined in VCam Input Settings
+	// Returns whether the profile was found and able to be applied
+	UFUNCTION(BlueprintCallable, Category = "VCam Input")
+	bool SetInputProfileFromName(const FName ProfileName);
+
+	// Tries to add a new Input Profile to the VCam Input Settings and populates it with any currently active player mappable keys
+	// Note: The set of currently active player mappable keys may be larger than the set of mappings in this component's Input Profile
+	//
+	// Optionally this function can update an existing profile, this will only add any mappable keys that don't currently exist in the profile but will not remove any existing mappings
+	// Returns whether the profile was successfully added or updated
+	UFUNCTION(BlueprintCallable, Category="VCam Input", meta=(ReturnDisplayName = "Success"))
+	bool AddInputProfileWithCurrentlyActiveMappings(const FName ProfileName, bool bUpdateIfProfileAlreadyExists = true);
+
+	// Saves the current input profile settings to the VCam Input Settings using the provided Profile Name
+	UFUNCTION(BlueprintCallable, Category="VCam Input", meta=(ReturnDisplayName = "Success"))
+	bool SaveCurrentInputProfileToSettings(const FName ProfileName) const;
+	
+	// Searches the currently active input system for all registered key mappings that are marked as Player Mappable.
+	UFUNCTION(BlueprintCallable, Category="VCam Input", meta=(ReturnDisplayName = "PlayerMappableActionKeyMappings"))
+	TArray<FEnhancedActionKeyMapping> GetAllPlayerMappableActionKeyMappings() const;
+
+	// Searches the currently active input system for the current key mapped to a given input mapping
+	// If there is not a player mapped key, then this will return EKeys::Invalid.
+	UFUNCTION(BlueprintCallable, Category="VCam Input", meta=(ReturnDisplayName = "Key"))
+	FKey GetPlayerMappedKey(const FName MappingName) const;
+
 	// Sets if the VCamComponent will update every frame or not
 	UFUNCTION(BlueprintSetter)
 	void SetEnabled(bool bNewEnabled);
@@ -253,6 +282,15 @@ public:
 	UPROPERTY(EditAnywhere, BlueprintReadWrite, Category = "VirtualCamera")
 	EVCamTargetViewportID TargetViewport = EVCamTargetViewportID::CurrentlySelected;
 
+	UPROPERTY(EditAnywhere, BlueprintReadWrite, BlueprintSetter=SetInputProfile, Category = "VirtualCamera")
+	FVCamInputProfile InputProfile;
+
+	// Call this after modifying the InputProfile in code to update the player mapped keys
+	void ApplyInputProfile();
+	
+	UFUNCTION(BlueprintCallable, BlueprintInternalUseOnly, Category="VirtualCamera")
+	void SetInputProfile(const FVCamInputProfile& NewInputProfile);
+
 	// List of Output Providers (executed in order)
 	UPROPERTY(EditAnywhere, Instanced, Category="VirtualCamera")
 	TArray<TObjectPtr<UVCamOutputProviderBase>> OutputProviders;
@@ -268,22 +306,14 @@ public:
 	UFUNCTION(BlueprintCallable, Category="VirtualCamera")
 	void RegisterObjectForInput(UObject* Object);
 
+	void UnregisterObjectForInput(UObject* Object) const;
+
 	/**
 	 * Returns a list of all player mappable keys that have been registered
 	 */
 	UFUNCTION(BlueprintCallable, Category="VirtualCamera")
 	TArray<FEnhancedActionKeyMapping> GetPlayerMappableKeys() const;
 
-	/**
-	 * Replace any currently applied mappings to this key mapping with the given new one.
-	 * Requests a rebuild of the player mappings.
-	 * This function will automatically switch between Editor based and Player based input systems as needed
-	 *
-	 * @return The number of mappings that have been replaced
-	 */
-	UFUNCTION(BlueprintCallable, Category = "VirtualCamera", meta=(AutoCreateRefTerm = "Options"))
-	int32 AddPlayerMappedKey(const FName MappingName, const FKey NewKey, const FModifyContextOptions& Options = FModifyContextOptions()) const;
-
 	/**
 	* Injects an input action. 
 	*/
@@ -297,7 +327,7 @@ public:
 	UFUNCTION(BlueprintCallable, Category = "VirtualCamera", meta = (AutoCreateRefTerm = "Modifiers,Triggers"))
 	virtual void InjectInputVectorForAction(const UInputAction* Action, FVector Value, const TArray<UInputModifier*>& Modifiers, const TArray<UInputTrigger*>& Triggers);
 
-private:
+private:	
 	static void CopyLiveLinkDataToCamera(const FLiveLinkCameraBlueprintData& LiveLinkData, UCineCameraComponent* CameraComponent);
 
 	float GetDeltaTime();
diff --git a/Engine/Plugins/Experimental/VirtualCamera/Source/VCamCore/VCamCore.Build.cs b/Engine/Plugins/Experimental/VirtualCamera/Source/VCamCore/VCamCore.Build.cs
--- a/Engine/Plugins/Experimental/VirtualCamera/Source/VCamCore/VCamCore.Build.cs
+++ b/Engine/Plugins/Experimental/VirtualCamera/Source/VCamCore/VCamCore.Build.cs
@@ -28,7 +28,8 @@ public class VCamCore : ModuleRules
 				"Composure",
 				"Core",
 				"VPUtilities",
-				"EnhancedInput"
+				"EnhancedInput",
+				"VCamInput"
 				// ... add other public dependencies that you statically link with here ...
 			}
 			);
@@ -48,7 +49,7 @@ public class VCamCore : ModuleRules
 				"MediaIOCore",
 				"RemoteSession",
 				"InputCore",
-				"VPRoles",
+				"VPRoles"
 				// ... add private dependencies that you statically link with here ...
 			}
 			);
diff --git a/Engine/Plugins/Experimental/VirtualCamera/Source/VCamCoreEditor/Private/VCamCoreEditorModule.cpp b/Engine/Plugins/Experimental/VirtualCamera/Source/VCamCoreEditor/Private/VCamCoreEditorModule.cpp
--- a/Engine/Plugins/Experimental/VirtualCamera/Source/VCamCoreEditor/Private/VCamCoreEditorModule.cpp
+++ b/Engine/Plugins/Experimental/VirtualCamera/Source/VCamCoreEditor/Private/VCamCoreEditorModule.cpp
@@ -4,7 +4,9 @@
 #include "Modules/ModuleManager.h"
 #include "Interfaces/IPluginManager.h"
 #include "Kismet2/KismetEditorUtilities.h"
+#include "PropertyEditorModule.h"
 #include "VCamModifier.h"
+#include "VCamInputProfileCustomization.h"
 
 class FVCamCoreEditorModule : public IModuleInterface
 {
@@ -15,11 +17,15 @@ public:
 	virtual void StartupModule() override
 	{
 		RegisterAutoGeneratedDefaultEvents();
+		
+		FPropertyEditorModule& PropertyModule = FModuleManager::GetModuleChecked<FPropertyEditorModule>("PropertyEditor");
+		PropertyModule.RegisterCustomPropertyTypeLayout("VCamInputProfile", FOnGetPropertyTypeCustomizationInstance::CreateStatic(&FVCamInputProfileCustomization::MakeInstance));
 	};
 	
 	virtual void ShutdownModule() override
 	{
-		
+		FPropertyEditorModule& PropertyModule = FModuleManager::GetModuleChecked<FPropertyEditorModule>("PropertyEditor");
+		PropertyModule.UnregisterCustomPropertyTypeLayout("VCamInputProfile");
 	};
 
 	// Set up default nodes for VCam Core blueprint classes
diff --git a/Engine/Plugins/Experimental/VirtualCamera/Source/VCamCoreEditor/Private/VCamInputProfileCustomization.cpp b/Engine/Plugins/Experimental/VirtualCamera/Source/VCamCoreEditor/Private/VCamInputProfileCustomization.cpp
new file mode 100644
--- /dev/null
+++ b/Engine/Plugins/Experimental/VirtualCamera/Source/VCamCoreEditor/Private/VCamInputProfileCustomization.cpp
@@ -0,0 +1,205 @@
+// Copyright Epic Games, Inc. All Rights Reserved.
+
+#include "VCamInputProfileCustomization.h"
+
+#include "DetailLayoutBuilder.h"
+#include "DetailWidgetRow.h"
+#include "IDetailChildrenBuilder.h"
+#include "VCamInputSettings.h"
+#include "ScopedTransaction.h"
+
+#define LOCTEXT_NAMESPACE "VCamInputProfileCustomization"
+
+FVCamInputProfileCustomization::FVCamInputProfileCustomization()
+{
+	BuildProfileComboList();
+}
+
+TSharedRef<IPropertyTypeCustomization> FVCamInputProfileCustomization::MakeInstance()
+{
+	return MakeShared<FVCamInputProfileCustomization>();
+}
+
+void FVCamInputProfileCustomization::CustomizeHeader(TSharedRef<IPropertyHandle> PropertyHandle, FDetailWidgetRow& HeaderRow, IPropertyTypeCustomizationUtils& CustomizationUtils)
+{
+	HeaderRow.
+		NameContent()
+		[
+			PropertyHandle->CreatePropertyNameWidget()
+		];
+
+	// We only want the dropdown list outside of the settings class as the settings class is the thing
+	// defining the presets we use for the dropdown
+	const bool bInSettingsClass = PropertyHandle->GetOuterBaseClass() == UVCamInputSettings::StaticClass();
+
+	if (!bInSettingsClass)
+	{
+		HeaderRow.
+			ValueContent()
+			.MaxDesiredWidth(0.f)
+			[
+				SAssignNew(ProfileComboBox, SComboBox< TSharedPtr<FString> >)
+				.OptionsSource(&ProfileComboList)
+				.OnGenerateWidget(this, &FVCamInputProfileCustomization::MakeProfileComboWidget)
+				.OnSelectionChanged(this, &FVCamInputProfileCustomization::OnProfileChanged)
+				.OnComboBoxOpening(this, &FVCamInputProfileCustomization::BuildProfileComboList)
+				.IsEnabled(this, &FVCamInputProfileCustomization::IsProfileEnabled)
+				.ContentPadding(2)
+				.Content()
+				[
+					SNew(STextBlock)
+					.Text(this, &FVCamInputProfileCustomization::GetProfileComboBoxContent)
+					.Font(IDetailLayoutBuilder::GetDetailFont())
+					.ToolTipText(this, &FVCamInputProfileCustomization::GetProfileComboBoxContent)
+				]
+			];
+	}
+}
+
+void FVCamInputProfileCustomization::CustomizeChildren(TSharedRef<IPropertyHandle> StructPropertyHandle, IDetailChildrenBuilder& ChildBuilder, IPropertyTypeCustomizationUtils& StructCustomizationUtils)
+{
+	CachedStructPropertyHandle = StructPropertyHandle;
+	
+	// Retrieve structure's child properties
+	uint32 NumChildren;
+	StructPropertyHandle->GetNumChildren( NumChildren );	
+	const FName MappableKeyOverridesPropertyName = GET_MEMBER_NAME_CHECKED(FVCamInputProfile, MappableKeyOverrides);
+	for(uint32 ChildIndex = 0; ChildIndex < NumChildren; ++ChildIndex)
+	{
+		TSharedPtr<IPropertyHandle> ChildProperty = StructPropertyHandle->GetChildHandle(ChildIndex);
+		ChildBuilder.AddProperty(ChildProperty.ToSharedRef());
+
+		if (ChildProperty->GetProperty() && ChildProperty->GetProperty()->GetFName() == MappableKeyOverridesPropertyName)
+		{
+			MappableKeyOverridesHandle = ChildProperty;
+		}
+	}
+}
+
+void FVCamInputProfileCustomization::BuildProfileComboList()
+{
+	const UVCamInputSettings* VCamInputSettings = GetDefault<UVCamInputSettings>();
+
+	if (!VCamInputSettings)
+	{
+		return;
+	}
+	
+	const TArray<FName>& Presets = VCamInputSettings->GetInputProfileNames();
+	
+	ProfileComboList.Empty(Presets.Num() + 1);
+	ProfileComboList.Add(MakeShared<FString>(TEXT("Custom...")));
+
+	// put all presets in the list
+	for (const FName& Preset : Presets)
+	{
+		ProfileComboList.Add(MakeShared<FString>(Preset.ToString()));
+	}
+}
+
+bool FVCamInputProfileCustomization::IsProfileEnabled() const
+{
+	if (MappableKeyOverridesHandle.IsValid())
+	{
+		return MappableKeyOverridesHandle->IsEditable() && FSlateApplication::Get().GetNormalExecutionAttribute().Get();
+	}
+	return false;
+}
+
+void FVCamInputProfileCustomization::OnProfileChanged(TSharedPtr<FString> NewSelection, ESelectInfo::Type SelectInfo)
+{
+	const UVCamInputSettings* VCamInputSettings = GetDefault<UVCamInputSettings>();
+
+	if (!VCamInputSettings)
+	{
+		return;
+	}
+	
+	// Ignore any changes that come from code
+	if (SelectInfo != ESelectInfo::Direct && CachedStructPropertyHandle.IsValid() && MappableKeyOverridesHandle.IsValid())
+	{
+		const FString NewPresetName = *NewSelection.Get();
+
+		if (const FVCamInputProfile* InputProfile = VCamInputSettings->InputProfiles.Find(FName(NewPresetName)))
+		{
+			const FScopedTransaction Transaction(LOCTEXT("ChangeVCamInputProfile", "Change VCam Input Profile"));
+			
+			FStructProperty* StructProperty = CastFieldChecked<FStructProperty>(CachedStructPropertyHandle->GetProperty());
+			FMapProperty* MapProperty = CastFieldChecked<FMapProperty>(MappableKeyOverridesHandle->GetProperty());
+
+			TArray<UObject*> OuterObjects;
+			MappableKeyOverridesHandle->GetOuterObjects(OuterObjects);
+
+			CachedStructPropertyHandle->NotifyPreChange();
+			
+			for (UObject* OuterObject : OuterObjects)
+			{				
+				void* StructValuePtr = StructProperty->ContainerPtrToValuePtr<void*>(OuterObject);
+				void* MapValuePtr = MapProperty->ContainerPtrToValuePtr<void*>(StructValuePtr);
+				
+				FScriptMapHelper MapHelper(MapProperty, MapValuePtr);
+
+				MapHelper.EmptyValues();
+
+				for (const TPair<FName, FKey>& MappingPair : InputProfile->MappableKeyOverrides)
+				{
+					MapHelper.AddPair(&MappingPair.Key, &MappingPair.Value);
+				}
+			}
+
+			CachedStructPropertyHandle->NotifyPostChange(EPropertyChangeType::ValueSet);
+			CachedStructPropertyHandle->NotifyFinishedChangingProperties();
+		}
+	}
+}
+
+TSharedRef<SWidget> FVCamInputProfileCustomization::MakeProfileComboWidget(TSharedPtr<FString> InItem)
+{
+	return SNew(STextBlock)
+		.Text(FText::FromString(*InItem))
+		.Font(IDetailLayoutBuilder::GetDetailFont());
+}
+
+FText FVCamInputProfileCustomization::GetProfileComboBoxContent() const
+{
+	// just test one variable for multiple selection
+	void* StructValueData;
+	if (CachedStructPropertyHandle->GetValueData(StructValueData) == FPropertyAccess::Result::MultipleValues)
+	{
+		// multiple selection
+		return LOCTEXT("MultipleValues", "Multiple Values");
+	}
+	const TSharedPtr<FString> ProfileString = GetProfileString(StructValueData);
+	return ProfileString.IsValid() ? FText::AsCultureInvariant(*ProfileString) : LOCTEXT("NoProfileFound", "Error");
+}
+
+TSharedPtr<FString> FVCamInputProfileCustomization::GetProfileString(void* StructValueData) const
+{
+	FVCamInputProfile* InputProfilePtr = static_cast<FVCamInputProfile*>(StructValueData);
+	const UVCamInputSettings* VCamInputSettings = GetDefault<UVCamInputSettings>();
+
+	if (InputProfilePtr && VCamInputSettings)
+	{
+		for (const TPair<FName, FVCamInputProfile>& InputProfilePair : VCamInputSettings->InputProfiles)
+		{
+			const FName& InputProfileName = InputProfilePair.Key;
+			const FVCamInputProfile& InputProfileValue = InputProfilePair.Value;
+			if (InputProfileValue == *InputProfilePtr)
+			{
+				const TSharedPtr<FString>* FoundString = ProfileComboList.FindByPredicate([InputProfileString = InputProfileName.ToString()](const TSharedPtr<FString> ListEntry)
+				{
+					return ListEntry.IsValid() ? *ListEntry == InputProfileString : false;
+				});
+				
+				if (FoundString)
+				{
+					return *FoundString;
+				}
+				break;
+			}
+		}
+	}
+	return ProfileComboList[0];
+}
+
+#undef LOCTEXT_NAMESPACE
\ No newline at end of file
diff --git a/Engine/Plugins/Experimental/VirtualCamera/Source/VCamCoreEditor/Public/VCamInputProfileCustomization.h b/Engine/Plugins/Experimental/VirtualCamera/Source/VCamCoreEditor/Public/VCamInputProfileCustomization.h
new file mode 100644
--- /dev/null
+++ b/Engine/Plugins/Experimental/VirtualCamera/Source/VCamCoreEditor/Public/VCamInputProfileCustomization.h
@@ -0,0 +1,43 @@
+// Copyright Epic Games, Inc. All Rights Reserved.
+
+#pragma once
+
+#include "Containers/Array.h"
+#include "Containers/UnrealString.h"
+#include "IPropertyTypeCustomization.h"
+#include "Internationalization/Text.h"
+#include "Templates/SharedPointer.h"
+#include "Types/SlateEnums.h"
+#include "Widgets/Input/SComboBox.h"
+
+class FDetailWidgetRow;
+class IPropertyHandle;
+class SWidget;
+
+class FVCamInputProfileCustomization : public IPropertyTypeCustomization
+{
+public:
+	FVCamInputProfileCustomization();
+
+	static TSharedRef<IPropertyTypeCustomization> MakeInstance();
+
+	/** IPropertyTypeCustomization instance */
+	virtual void CustomizeHeader(TSharedRef<IPropertyHandle> PropertyHandle, FDetailWidgetRow& HeaderRow, IPropertyTypeCustomizationUtils& CustomizationUtils) override;
+	virtual void CustomizeChildren( TSharedRef<IPropertyHandle> StructPropertyHandle, class IDetailChildrenBuilder& ChildBuilder, IPropertyTypeCustomizationUtils& StructCustomizationUtils ) override;
+
+protected:
+	TSharedPtr<IPropertyHandle> CachedStructPropertyHandle;
+	TSharedPtr<IPropertyHandle> MappableKeyOverridesHandle;
+	
+	TSharedPtr<SComboBox<TSharedPtr<FString>>> ProfileComboBox;
+	TArray<TSharedPtr<FString>> ProfileComboList;
+
+	void BuildProfileComboList();
+
+	bool IsProfileEnabled() const;
+	void OnProfileChanged(TSharedPtr<FString> NewSelection, ESelectInfo::Type SelectInfo);
+	TSharedRef<SWidget> MakeProfileComboWidget(TSharedPtr<FString> InItem);
+		
+	FText GetProfileComboBoxContent() const;
+	TSharedPtr<FString> GetProfileString(void* StructValueData) const;
+};
diff --git a/Engine/Plugins/Experimental/VirtualCamera/Source/VCamCoreEditor/VCamCoreEditor.build.cs b/Engine/Plugins/Experimental/VirtualCamera/Source/VCamCoreEditor/VCamCoreEditor.build.cs
--- a/Engine/Plugins/Experimental/VirtualCamera/Source/VCamCoreEditor/VCamCoreEditor.build.cs
+++ b/Engine/Plugins/Experimental/VirtualCamera/Source/VCamCoreEditor/VCamCoreEditor.build.cs
@@ -31,7 +31,9 @@ public class VCamCoreEditor : ModuleRules
 				"EnhancedInput",
 				"PropertyEditor",
 				"UMG", 
-				"UMGEditor"
+				"UMGEditor",
+				"InputCore",
+				"VCamInput"
 			}
 		);
 	}
diff --git a/Engine/Plugins/Experimental/VirtualCamera/Source/VCamInput/Private/VCamInputSettings.cpp b/Engine/Plugins/Experimental/VirtualCamera/Source/VCamInput/Private/VCamInputSettings.cpp
new file mode 100644
--- /dev/null
+++ b/Engine/Plugins/Experimental/VirtualCamera/Source/VCamInput/Private/VCamInputSettings.cpp
@@ -0,0 +1,48 @@
+// Copyright Epic Games, Inc. All Rights Reserved.
+
+#include "VCamInputSettings.h"
+
+bool FVCamInputProfile::operator==(const FVCamInputProfile& OtherProfile) const
+{
+	bool bIsEqual = MappableKeyOverrides.Num() == OtherProfile.MappableKeyOverrides.Num();
+	if (bIsEqual)
+	{
+		for (const TPair<FName, FKey>& MappingPair : MappableKeyOverrides)
+		{
+			const FKey* OtherKey = OtherProfile.MappableKeyOverrides.Find(MappingPair.Key);
+			if (!OtherKey || MappingPair.Value != *OtherKey)
+			{
+				bIsEqual = false;
+				break;
+			}
+		}
+	}
+	return bIsEqual;
+}
+
+void UVCamInputSettings::SetDefaultInputProfile(const FName NewDefaultInputProfile)
+{
+	if (InputProfiles.Contains(NewDefaultInputProfile))
+	{
+		DefaultInputProfile = NewDefaultInputProfile;
+		SaveConfig();
+	}
+}
+
+void UVCamInputSettings::SetInputProfiles(const TMap<FName, FVCamInputProfile>& NewInputProfiles)
+{
+	InputProfiles = NewInputProfiles;
+	SaveConfig();
+}
+
+TArray<FName> UVCamInputSettings::GetInputProfileNames() const
+{
+	TArray<FName> ProfileNames;
+	InputProfiles.GenerateKeyArray(ProfileNames);
+	return ProfileNames;
+}
+
+UVCamInputSettings* UVCamInputSettings::GetVCamInputSettings()
+{
+	return GetMutableDefault<UVCamInputSettings>();
+}
diff --git a/Engine/Plugins/Experimental/VirtualCamera/Source/VCamInput/Public/VCamInputSettings.h b/Engine/Plugins/Experimental/VirtualCamera/Source/VCamInput/Public/VCamInputSettings.h
new file mode 100644
--- /dev/null
+++ b/Engine/Plugins/Experimental/VirtualCamera/Source/VCamInput/Public/VCamInputSettings.h
@@ -0,0 +1,47 @@
+// Copyright Epic Games, Inc. All Rights Reserved.
+
+#pragma once
+
+#include "CoreMinimal.h"
+
+#include "Engine/DeveloperSettings.h"
+#include "InputCoreTypes.h"
+
+#include "VCamInputSettings.generated.h"
+
+USTRUCT(BlueprintType)
+struct VCAMINPUT_API FVCamInputProfile
+{
+	GENERATED_BODY()
+
+	UPROPERTY(EditAnywhere, BlueprintReadWrite, Category="Input")
+	TMap<FName, FKey> MappableKeyOverrides;
+
+	bool operator==(const FVCamInputProfile& OtherProfile) const;
+};
+
+UCLASS(config=Game)
+class VCAMINPUT_API UVCamInputSettings : public UDeveloperSettings
+{
+	GENERATED_BODY()
+
+public:
+	UPROPERTY(EditAnywhere, Config, BlueprintReadWrite, BlueprintSetter=SetDefaultInputProfile, Category="Input", meta=(GetOptions=GetInputProfileNames))
+	FName DefaultInputProfile;
+
+	UFUNCTION(BlueprintCallable, Category="VCam Input")
+	void SetDefaultInputProfile(const FName NewDefaultInputProfile);
+	
+	UPROPERTY(EditAnywhere, Config, BlueprintReadWrite, BlueprintSetter=SetInputProfiles, Category="Input")
+	TMap<FName, FVCamInputProfile> InputProfiles;
+
+	UFUNCTION(BlueprintCallable, Category="VCam Input")
+	void SetInputProfiles(const TMap<FName, FVCamInputProfile>& NewInputProfiles);
+
+	UFUNCTION(BlueprintCallable, Category="VCam Input")
+	TArray<FName> GetInputProfileNames() const;
+	
+private:
+	UFUNCTION(BlueprintCallable, Category="VCam")
+	static UVCamInputSettings* GetVCamInputSettings();
+};
\ No newline at end of file
diff --git a/Engine/Plugins/Experimental/VirtualCamera/Source/VCamInput/VCamInput.build.cs b/Engine/Plugins/Experimental/VirtualCamera/Source/VCamInput/VCamInput.build.cs
--- a/Engine/Plugins/Experimental/VirtualCamera/Source/VCamInput/VCamInput.build.cs
+++ b/Engine/Plugins/Experimental/VirtualCamera/Source/VCamInput/VCamInput.build.cs
@@ -26,7 +26,8 @@ public class VCamInput : ModuleRules
 				"Engine",
 				"InputCore",
 				"Slate",
-				"SlateCore"
+				"SlateCore",
+				"DeveloperSettings"
 			}
 			);
 	}
