Add CineCameraComponent in VPBookmarkActor

diff --git a/Engine/Plugins/Experimental/VirtualCamera/Source/VirtualCamera/Private/VCamBlueprintFunctionLibrary.cpp b/Engine/Plugins/Experimental/VirtualCamera/Source/VirtualCamera/Private/VCamBlueprintFunctionLibrary.cpp
--- a/Engine/Plugins/Experimental/VirtualCamera/Source/VirtualCamera/Private/VCamBlueprintFunctionLibrary.cpp
+++ b/Engine/Plugins/Experimental/VirtualCamera/Source/VirtualCamera/Private/VCamBlueprintFunctionLibrary.cpp
@@ -424,6 +424,16 @@ bool UVCamBlueprintFunctionLibrary::CopyToCineCameraActor(UCineCameraComponent*
 	return true;
 }
 
+void UVCamBlueprintFunctionLibrary::SetActorLabel(AActor* Actor, const FString& NewActorLabel)
+{
+#if WITH_EDITOR
+	if (Actor)
+	{
+		Actor->SetActorLabel(NewActorLabel);
+	}
+#endif
+}
+
 bool UVCamBlueprintFunctionLibrary::DeprojectScreenToWorld(const FVector2D& InScreenPosition, FVector& OutWorldPosition, FVector& OutWorldDirection)
 {
 	FName LevelEditorName(TEXT("LevelEditor"));
diff --git a/Engine/Plugins/Experimental/VirtualCamera/Source/VirtualCamera/Public/VCamBlueprintFunctionLibrary.h b/Engine/Plugins/Experimental/VirtualCamera/Source/VirtualCamera/Public/VCamBlueprintFunctionLibrary.h
--- a/Engine/Plugins/Experimental/VirtualCamera/Source/VirtualCamera/Public/VCamBlueprintFunctionLibrary.h
+++ b/Engine/Plugins/Experimental/VirtualCamera/Source/VirtualCamera/Public/VCamBlueprintFunctionLibrary.h
@@ -147,6 +147,9 @@ public:
 	UFUNCTION(BlueprintCallable, Category = "VirtualCamera")
 	static bool CopyToCineCameraActor(UCineCameraComponent* SourceCameraComponent, ACineCameraActor* TargetCameraActor);
 
+	UFUNCTION(BlueprintCallable, Category = "VirtualCamera")
+	static void SetActorLabel(AActor* Actor, const FString& NewActorLabel);
+
 private:
 
 	static bool DeprojectScreenToWorld(const FVector2D& InScreenPosition, FVector& OutWorldPosition, FVector& OutWorldDirection);
