Changed VirtualCamera plugin to Runtime

diff --git a/Engine/Plugins/Experimental/VirtualCamera/Source/VirtualCamera/Private/LevelSequencePlaybackController.cpp b/Engine/Plugins/Experimental/VirtualCamera/Source/VirtualCamera/Private/LevelSequencePlaybackController.cpp
--- a/Engine/Plugins/Experimental/VirtualCamera/Source/VirtualCamera/Private/LevelSequencePlaybackController.cpp
+++ b/Engine/Plugins/Experimental/VirtualCamera/Source/VirtualCamera/Private/LevelSequencePlaybackController.cpp
@@ -7,6 +7,7 @@
 #include "HAL/FileManager.h"
 #include "LevelSequence.h"
 #include "Modules/ModuleManager.h"
+#include "MovieScene.h"
 #include "MovieSceneTimeHelpers.h"
 #include "UObject/UObjectBase.h"
 #include "UObject/UObjectIterator.h"
diff --git a/Engine/Plugins/Experimental/VirtualCamera/Source/VirtualCamera/Private/VCamBlueprintFunctionLibrary.cpp b/Engine/Plugins/Experimental/VirtualCamera/Source/VirtualCamera/Private/VCamBlueprintFunctionLibrary.cpp
--- a/Engine/Plugins/Experimental/VirtualCamera/Source/VirtualCamera/Private/VCamBlueprintFunctionLibrary.cpp
+++ b/Engine/Plugins/Experimental/VirtualCamera/Source/VirtualCamera/Private/VCamBlueprintFunctionLibrary.cpp
@@ -10,6 +10,8 @@
 #include "AssetRegistry/AssetData.h"
 #include "VirtualCameraUserSettings.h"
 #include "GameFramework/PlayerController.h"
+#include "MovieScene.h"
+
 #if WITH_EDITOR
 #include "Editor.h"
 #include "Editor/Transactor.h"
@@ -25,6 +27,8 @@
 #include "Subsystems/UnrealEditorSubsystem.h"
 #include "LevelEditor/Public/LevelEditorSubsystem.h"
 #include "Modules/ModuleManager.h"
+#include "Recorder/TakeRecorderBlueprintLibrary.h"
+#include "Recorder/TakeRecorderPanel.h"
 #endif
 
 bool UVCamBlueprintFunctionLibrary::IsGameRunning()
@@ -383,11 +387,13 @@ void UVCamBlueprintFunctionLibrary::EditorSetGameView(bool bIsToggled)
 
 void UVCamBlueprintFunctionLibrary::EnableDebugFocusPlane(UCineCameraComponent* CineCamera, bool bEnabled)
 {
+#if WITH_EDITOR
 	if (!CineCamera)
 	{
 		return;
 	}
 	CineCamera->FocusSettings.bDrawDebugFocusPlane = bEnabled;
+#endif
 }
 
 FString UVCamBlueprintFunctionLibrary::GetNextUndoDescription()
@@ -434,6 +440,24 @@ void UVCamBlueprintFunctionLibrary::SetActorLabel(AActor* Actor, const FString&
 #endif
 }
 
+bool UVCamBlueprintFunctionLibrary::IsTakeRecorderPanelOpen()
+{
+#if WITH_EDITOR
+	return IsValid(UTakeRecorderBlueprintLibrary::GetTakeRecorderPanel());
+#else
+	return false;
+#endif
+}
+
+bool UVCamBlueprintFunctionLibrary::TryOpenTakeRecorderPanel()
+{
+#if WITH_EDITOR
+	return IsValid(UTakeRecorderBlueprintLibrary::OpenTakeRecorderPanel());
+#else
+	return false;
+#endif
+}
+
 bool UVCamBlueprintFunctionLibrary::DeprojectScreenToWorld(const FVector2D& InScreenPosition, FVector& OutWorldPosition, FVector& OutWorldDirection)
 {
 	FName LevelEditorName(TEXT("LevelEditor"));
diff --git a/Engine/Plugins/Experimental/VirtualCamera/Source/VirtualCamera/Private/VirtualCameraActor.cpp b/Engine/Plugins/Experimental/VirtualCamera/Source/VirtualCamera/Private/VirtualCameraActor.cpp
--- a/Engine/Plugins/Experimental/VirtualCamera/Source/VirtualCamera/Private/VirtualCameraActor.cpp
+++ b/Engine/Plugins/Experimental/VirtualCamera/Source/VirtualCamera/Private/VirtualCameraActor.cpp
@@ -268,6 +268,7 @@ void AVirtualCameraActor::SetTrackedActorForFocus_Implementation(AActor* InActor
 
 void AVirtualCameraActor::SetFocusVisualization_Implementation(bool bInShowFocusVisualization)
 {
+#if WITH_EDITOR
 	UCineCameraComponent* CineCamera = GetCineCameraComponent();
 	if (CineCamera)
 	{
@@ -278,6 +279,7 @@ void AVirtualCameraActor::SetFocusVisualization_Implementation(bool bInShowFocus
 		}
 		CineCamera->FocusSettings.bDrawDebugFocusPlane = bInShowFocusVisualization;
 	}
+#endif
 }
 
 void AVirtualCameraActor::SetBeforeSetVirtualCameraTransformDelegate_Implementation(const FPreSetVirtualCameraTransform& InDelegate)
@@ -663,7 +665,11 @@ void AVirtualCameraActor::SaveSettings()
 	SaveGameInstance->CameraSettings.FocalLength = CineCamera->CurrentFocalLength;
 	SaveGameInstance->CameraSettings.Aperture = CineCamera->CurrentAperture;
 	SaveGameInstance->CameraSettings.bAllowFocusVisualization = bAllowFocusVisualization;
+#if WITH_EDITORONLY_DATA
 	SaveGameInstance->CameraSettings.DebugFocusPlaneColor = CineCamera->FocusSettings.DebugFocusPlaneColor;
+#else
+	SaveGameInstance->CameraSettings.DebugFocusPlaneColor = FColor();
+#endif
 
 	// Save filmback settings
 	SaveGameInstance->CameraSettings.FilmbackName = CineCamera->GetFilmbackPresetName();
@@ -703,10 +709,12 @@ void AVirtualCameraActor::LoadSettings()
 
 	bAllowFocusVisualization = SaveGameInstance->CameraSettings.bAllowFocusVisualization;
 
+#if WITH_EDITORONLY_DATA
 	if (SaveGameInstance->CameraSettings.DebugFocusPlaneColor != FColor())
 	{
 		CineCamera->FocusSettings.DebugFocusPlaneColor = SaveGameInstance->CameraSettings.DebugFocusPlaneColor;
 	}
+#endif
 
 	CineCamera->SetCurrentFocalLength(SaveGameInstance->CameraSettings.FocalLength);
 	CineCamera->CurrentAperture = SaveGameInstance->CameraSettings.Aperture;
diff --git a/Engine/Plugins/Experimental/VirtualCamera/Source/VirtualCamera/Private/VirtualCameraCineCameraComponent.cpp b/Engine/Plugins/Experimental/VirtualCamera/Source/VirtualCamera/Private/VirtualCameraCineCameraComponent.cpp
--- a/Engine/Plugins/Experimental/VirtualCamera/Source/VirtualCamera/Private/VirtualCameraCineCameraComponent.cpp
+++ b/Engine/Plugins/Experimental/VirtualCamera/Source/VirtualCamera/Private/VirtualCameraCineCameraComponent.cpp
@@ -229,6 +229,7 @@ void UVirtualCameraCineCameraComponent::SetFocusChangeSmoothness(const float New
 
 void UVirtualCameraCineCameraComponent::SetFocusVisualization(bool bShowFocusVisualization)
 {
+#if WITH_EDITOR
 	if (FocusSettings.FocusMethod == ECameraFocusMethod::Disable)
 	{
 		UE_LOG(LogActor, Warning, TEXT("Camera focus mode is currently set to none, cannot display focus plane!"))
@@ -236,11 +237,16 @@ void UVirtualCameraCineCameraComponent::SetFocusVisualization(bool bShowFocusVis
 	}
 
 	FocusSettings.bDrawDebugFocusPlane = bShowFocusVisualization;
+#endif
 }
 
 bool UVirtualCameraCineCameraComponent::IsFocusVisualizationActivated() const
 {
+#if WITH_EDITOR
 	return FocusSettings.FocusMethod != ECameraFocusMethod::Disable && FocusSettings.bDrawDebugFocusPlane;
+#else
+	return false;
+#endif
 }
 
 int32 UVirtualCameraCineCameraComponent::FindClosestPresetIndex(const TArray<float>& ArrayToSearch, const float SearchValue) const
diff --git a/Engine/Plugins/Experimental/VirtualCamera/Source/VirtualCamera/Private/VirtualCameraClipsMetaData.cpp b/Engine/Plugins/Experimental/VirtualCamera/Source/VirtualCamera/Private/VirtualCameraClipsMetaData.cpp
--- a/Engine/Plugins/Experimental/VirtualCamera/Source/VirtualCamera/Private/VirtualCameraClipsMetaData.cpp
+++ b/Engine/Plugins/Experimental/VirtualCamera/Source/VirtualCamera/Private/VirtualCameraClipsMetaData.cpp
@@ -36,7 +36,7 @@ void UVirtualCameraClipsMetaData::ExtendAssetRegistryTags(TArray<FAssetRegistryT
 	OutTags.Emplace(AssetRegistryTag_bIsACineCameraRecording, IsRecordedFromACineCameraActor, FAssetRegistryTag::ETagType::TT_Alphabetical, FAssetRegistryTag::TD_None);
 }
 
-
+#if WITH_EDITOR
 void UVirtualCameraClipsMetaData::ExtendAssetRegistryTagMetaData(TMap<FName, FAssetRegistryTagMetadata>& OutMetadata) const
 {
 	OutMetadata.Add(AssetRegistryTag_FocalLength, FAssetRegistryTagMetadata()
@@ -76,9 +76,8 @@ void UVirtualCameraClipsMetaData::ExtendAssetRegistryTagMetaData(TMap<FName, FAs
 		.SetDisplayName(NSLOCTEXT("ClipsMetaData", "IsACineCameraRecording_Label", "bIsACineCameraRecording"))
 		.SetTooltip(NSLOCTEXT("ClipsMetaData", "IsACineCameraRecording_Tip", "If the level sequence was recorded with a CineCameraActor."))
 	);
-
-
 }
+#endif
 
 float UVirtualCameraClipsMetaData::GetFocalLength() const 
 {
diff --git a/Engine/Plugins/Experimental/VirtualCamera/Source/VirtualCamera/Private/VirtualCameraPawnBase.cpp b/Engine/Plugins/Experimental/VirtualCamera/Source/VirtualCamera/Private/VirtualCameraPawnBase.cpp
--- a/Engine/Plugins/Experimental/VirtualCamera/Source/VirtualCamera/Private/VirtualCameraPawnBase.cpp
+++ b/Engine/Plugins/Experimental/VirtualCamera/Source/VirtualCamera/Private/VirtualCameraPawnBase.cpp
@@ -424,7 +424,11 @@ void AVirtualCameraPawnBase::SaveSettings()
 	SaveGameInstance->CameraSettings.FocalLength = CineCamera->GetCurrentFocalLength();
 	SaveGameInstance->CameraSettings.Aperture = CineCamera->GetCurrentAperture();
 	SaveGameInstance->CameraSettings.bAllowFocusVisualization = bAllowFocusVisualization;
+#if WITH_EDITORONLY_DATA
 	SaveGameInstance->CameraSettings.DebugFocusPlaneColor = CineCamera->FocusSettings.DebugFocusPlaneColor;
+#else
+	SaveGameInstance->CameraSettings.DebugFocusPlaneColor = FColor();
+#endif
 
 	// Save filmback settings
 	SaveGameInstance->CameraSettings.FilmbackName = CineCamera->GetCurrentFilmbackName();
@@ -486,10 +490,12 @@ void AVirtualCameraPawnBase::LoadSettings()
 
 	bAllowFocusVisualization = SaveGameInstance->CameraSettings.bAllowFocusVisualization;
 
+#if WITH_EDITORONLY_DATA
 	if (SaveGameInstance->CameraSettings.DebugFocusPlaneColor != FColor())
 	{
 		CineCamera->FocusSettings.DebugFocusPlaneColor = SaveGameInstance->CameraSettings.DebugFocusPlaneColor;
 	}
+#endif
 
 	// Load focal length
 	if (CineCamera->FocalLengthOptions.Contains(SaveGameInstance->CameraSettings.FocalLength))
diff --git a/Engine/Plugins/Experimental/VirtualCamera/Source/VirtualCamera/Private/VirtualCameraPlayerControllerBase.cpp b/Engine/Plugins/Experimental/VirtualCamera/Source/VirtualCamera/Private/VirtualCameraPlayerControllerBase.cpp
--- a/Engine/Plugins/Experimental/VirtualCamera/Source/VirtualCamera/Private/VirtualCameraPlayerControllerBase.cpp
+++ b/Engine/Plugins/Experimental/VirtualCamera/Source/VirtualCamera/Private/VirtualCameraPlayerControllerBase.cpp
@@ -882,11 +882,12 @@ EUnit AVirtualCameraPlayerControllerBase::GetDesiredDistanceUnits()
 
 FColor AVirtualCameraPlayerControllerBase::GetFocusPlaneColor()
 {
+#if WITH_EDITOR
 	if (UVirtualCameraCineCameraComponent* CineCamera = GetVirtualCameraCineCameraComponent())
 	{
 		return CineCamera->FocusSettings.DebugFocusPlaneColor;
 	}
-
+#endif
 	return FColor();
 }
 
@@ -1248,10 +1249,12 @@ void AVirtualCameraPlayerControllerBase::SetFocusMethod(const EVirtualCameraFocu
 
 void AVirtualCameraPlayerControllerBase::SetFocusPlaneColor(const FColor NewFocusPlaneColor)
 {
+#if WITH_EDITOR
 	if (GetVirtualCameraCineCameraComponent())
 	{
 		GetVirtualCameraCineCameraComponent()->FocusSettings.DebugFocusPlaneColor = NewFocusPlaneColor;
 	}
+#endif
 }
 
 void AVirtualCameraPlayerControllerBase::SetFocusVisualization(bool bShowFocusVisualization)
diff --git a/Engine/Plugins/Experimental/VirtualCamera/Source/VirtualCamera/Private/VirtualCameraSubsystem.cpp b/Engine/Plugins/Experimental/VirtualCamera/Source/VirtualCamera/Private/VirtualCameraSubsystem.cpp
--- a/Engine/Plugins/Experimental/VirtualCamera/Source/VirtualCamera/Private/VirtualCameraSubsystem.cpp
+++ b/Engine/Plugins/Experimental/VirtualCamera/Source/VirtualCamera/Private/VirtualCameraSubsystem.cpp
@@ -9,8 +9,10 @@ UVirtualCameraSubsystem::UVirtualCameraSubsystem()
 	: bIsStreaming(false)
 {
 	SequencePlaybackController = CreateDefaultSubobject<ULevelSequencePlaybackController>("SequencePlaybackController");
+#if WITH_EDITOR
 	USelection::SelectionChangedEvent.AddUObject(this, &UVirtualCameraSubsystem::HandleSelectionChangedEvent);
 	USelection::SelectObjectEvent.AddUObject(this, &UVirtualCameraSubsystem::HandleSelectObjectEvent);
+#endif
 }
 
 bool UVirtualCameraSubsystem::StartStreaming()
@@ -60,7 +62,7 @@ bool UVirtualCameraSubsystem::IsStreaming() const
 	return bIsStreaming;
 }
 
-
+#if WITH_EDITOR
 void UVirtualCameraSubsystem::HandleSelectionChangedEvent(UObject* ChangedObject)
 {
 	USelection* Selection = Cast<USelection>(ChangedObject);
@@ -84,7 +86,7 @@ void UVirtualCameraSubsystem::HandleSelectObjectEvent(UObject* ChangedObject)
 		OnSelectedActorInViewportDelegate.Broadcast(SelectedActorCheck2);
 	}
 }
-
+#endif
 
 TScriptInterface<IVirtualCameraController> UVirtualCameraSubsystem::GetVirtualCameraController() const
 {
diff --git a/Engine/Plugins/Experimental/VirtualCamera/Source/VirtualCamera/Public/VCamBlueprintFunctionLibrary.h b/Engine/Plugins/Experimental/VirtualCamera/Source/VirtualCamera/Public/VCamBlueprintFunctionLibrary.h
--- a/Engine/Plugins/Experimental/VirtualCamera/Source/VirtualCamera/Public/VCamBlueprintFunctionLibrary.h
+++ b/Engine/Plugins/Experimental/VirtualCamera/Source/VirtualCamera/Public/VCamBlueprintFunctionLibrary.h
@@ -150,6 +150,11 @@ public:
 	UFUNCTION(BlueprintCallable, Category = "VirtualCamera")
 	static void SetActorLabel(AActor* Actor, const FString& NewActorLabel);
 
+	UFUNCTION(BlueprintPure, Category = "VirtualCamera")
+	static bool IsTakeRecorderPanelOpen();
+
+	UFUNCTION(BlueprintCallable, Category = "VirtualCamera")
+	static bool TryOpenTakeRecorderPanel();
 private:
 
 	static bool DeprojectScreenToWorld(const FVector2D& InScreenPosition, FVector& OutWorldPosition, FVector& OutWorldDirection);
diff --git a/Engine/Plugins/Experimental/VirtualCamera/Source/VirtualCamera/Public/VirtualCameraClipsMetaData.h b/Engine/Plugins/Experimental/VirtualCamera/Source/VirtualCamera/Public/VirtualCameraClipsMetaData.h
--- a/Engine/Plugins/Experimental/VirtualCamera/Source/VirtualCamera/Public/VirtualCameraClipsMetaData.h
+++ b/Engine/Plugins/Experimental/VirtualCamera/Source/VirtualCamera/Public/VirtualCameraClipsMetaData.h
@@ -54,8 +54,9 @@ public:
 	* Extend the default ULevelSequence asset registry tags
 	*/
 	virtual void ExtendAssetRegistryTags(TArray<FAssetRegistryTag>& OutTags) const override;
-
+#if WITH_EDITOR
 	virtual void ExtendAssetRegistryTagMetaData(TMap<FName, FAssetRegistryTagMetadata>& OutMetadata) const override;
+#endif
 
 public:
 	/**
diff --git a/Engine/Plugins/Experimental/VirtualCamera/Source/VirtualCamera/Public/VirtualCameraSubsystem.h b/Engine/Plugins/Experimental/VirtualCamera/Source/VirtualCamera/Public/VirtualCameraSubsystem.h
--- a/Engine/Plugins/Experimental/VirtualCamera/Source/VirtualCamera/Public/VirtualCameraSubsystem.h
+++ b/Engine/Plugins/Experimental/VirtualCamera/Source/VirtualCamera/Public/VirtualCameraSubsystem.h
@@ -31,11 +31,13 @@ public:
 	UFUNCTION(BlueprintCallable, Category = "VirtualCamera | Streaming")
 	bool IsStreaming() const;
 
+#if WITH_EDITOR
 	UFUNCTION()
 	void HandleSelectionChangedEvent(UObject* SelectedObject);
 
 	UFUNCTION()
 	void HandleSelectObjectEvent(UObject* SelectedObject);
+#endif
 
 	UFUNCTION(BlueprintCallable, Category = "VirtualCamera")
 	TScriptInterface<IVirtualCameraController> GetVirtualCameraController() const;
diff --git a/Engine/Plugins/Experimental/VirtualCamera/Source/VirtualCamera/VirtualCamera.Build.cs b/Engine/Plugins/Experimental/VirtualCamera/Source/VirtualCamera/VirtualCamera.Build.cs
--- a/Engine/Plugins/Experimental/VirtualCamera/Source/VirtualCamera/VirtualCamera.Build.cs
+++ b/Engine/Plugins/Experimental/VirtualCamera/Source/VirtualCamera/VirtualCamera.Build.cs
@@ -25,8 +25,6 @@ public class VirtualCamera : ModuleRules
 				"UMG",
 				"VCamCore",
 				"VPUtilities",
-				"VPUtilitiesEditor",
-				"TakesCore",
 				"AssetRegistry",
 			}
 		);
@@ -68,6 +66,8 @@ public class VirtualCamera : ModuleRules
 			PrivateDependencyModuleNames.Add("EditorFramework");
 			PrivateDependencyModuleNames.Add("UnrealEd");
 			PrivateDependencyModuleNames.Add("EditorScriptingUtilities");
+			PrivateDependencyModuleNames.Add("VPUtilitiesEditor");
+			PrivateDependencyModuleNames.Add("TakesCore");
 		}
 	}
 }
diff --git a/Engine/Plugins/Experimental/VirtualCamera/VirtualCamera.uplugin b/Engine/Plugins/Experimental/VirtualCamera/VirtualCamera.uplugin
--- a/Engine/Plugins/Experimental/VirtualCamera/VirtualCamera.uplugin
+++ b/Engine/Plugins/Experimental/VirtualCamera/VirtualCamera.uplugin
@@ -41,7 +41,7 @@
 		},
 		{
 			"Name": "VirtualCamera",
-			"Type": "UncookedOnly",
+			"Type": "Runtime",
 			"LoadingPhase": "PostDefault",
 			"PlatformAllowList": [
 				"Win64",
