Fix for VCam Modifier Names resetting when changing modifier properties in the details panel

diff --git a/Engine/Plugins/Experimental/VirtualCamera/Source/VCamCore/Private/VCamComponent.cpp b/Engine/Plugins/Experimental/VirtualCamera/Source/VCamCore/Private/VCamComponent.cpp
--- a/Engine/Plugins/Experimental/VirtualCamera/Source/VCamCore/Private/VCamComponent.cpp
+++ b/Engine/Plugins/Experimental/VirtualCamera/Source/VCamCore/Private/VCamComponent.cpp
@@ -268,42 +268,44 @@ void UVCamComponent::CheckForErrors()
 
 void UVCamComponent::PreEditChange(FProperty* PropertyThatWillChange)
 {
+	Super::PreEditChange(PropertyThatWillChange);
+}
+
+void UVCamComponent::PreEditChange(FEditPropertyChain& PropertyAboutToChange)
+{
+	FProperty* MemberProperty = PropertyAboutToChange.GetActiveMemberNode()->GetValue();
+
 	// Copy the property that is going to be changed so we can use it in PostEditChange if needed (for ArrayClear, ArrayRemove, etc.)
-	if (PropertyThatWillChange)
+	if (MemberProperty)
 	{
 		static FName NAME_OutputProviders = GET_MEMBER_NAME_CHECKED(UVCamComponent, OutputProviders);
 		static FName NAME_ModifierStack = GET_MEMBER_NAME_CHECKED(UVCamComponent, ModifierStack);
-		// Name property withing the Modifier Stack Entry struct. Possible collision due to just being called "Name"
-		static FName NAME_ModifierStackEntryName = GET_MEMBER_NAME_CHECKED(FModifierStackEntry, Name);
 		static FName NAME_Enabled = GET_MEMBER_NAME_CHECKED(UVCamComponent, bEnabled);
 
-		const FName PropertyThatWillChangeName = PropertyThatWillChange->GetFName();
+		const FName MemberPropertyName = MemberProperty->GetFName();
 
-		if (PropertyThatWillChangeName == NAME_OutputProviders)
+		if (MemberPropertyName == NAME_OutputProviders)
 		{
 			SavedOutputProviders.Empty();
 			SavedOutputProviders = OutputProviders;
 		}
-		else if (PropertyThatWillChangeName == NAME_ModifierStack || PropertyThatWillChangeName == NAME_ModifierStackEntryName)
+		else if (MemberPropertyName == NAME_ModifierStack)
 		{
 			SavedModifierStack = ModifierStack;
 		}
-		else if (PropertyThatWillChangeName == NAME_Enabled)
+		else if (MemberPropertyName == NAME_Enabled)
 		{
-			// If the property's owner is a struct (like FModifierStackEntry), act on it in PostEditChangeProperty(), not here
-			if (PropertyThatWillChange->GetOwner<UClass>())
-			{
-				void* PropertyData = PropertyThatWillChange->ContainerPtrToValuePtr<void>(this);
-				bool bWasEnabled = false;
-				PropertyThatWillChange->CopySingleValue(&bWasEnabled, PropertyData);
-
-				// Changing the enabled state needs to be done here instead of PostEditChange
-				SetEnabled(!bWasEnabled);
-			}
+			// Changing the enabled state needs to be done here instead of PostEditChange
+			// So we need to grab the value from the FProperty directly before using it
+			void* PropertyData = MemberProperty->ContainerPtrToValuePtr<void>(this);
+			bool bWasEnabled = false;
+			MemberProperty->CopySingleValue(&bWasEnabled, PropertyData);
+			
+			SetEnabled(!bWasEnabled);
 		}
 	}
-
-	Super::PreEditChange(PropertyThatWillChange);
+	
+	UObject::PreEditChange(PropertyAboutToChange);
 }
 
 void UVCamComponent::PostEditChangeProperty(FPropertyChangedEvent& PropertyChangedEvent)
diff --git a/Engine/Plugins/Experimental/VirtualCamera/Source/VCamCore/Public/VCamComponent.h b/Engine/Plugins/Experimental/VirtualCamera/Source/VCamCore/Public/VCamComponent.h
--- a/Engine/Plugins/Experimental/VirtualCamera/Source/VCamCore/Public/VCamComponent.h
+++ b/Engine/Plugins/Experimental/VirtualCamera/Source/VCamCore/Public/VCamComponent.h
@@ -67,6 +67,7 @@ public:
 
 	virtual void CheckForErrors() override;
 	virtual void PreEditChange(FProperty* PropertyThatWillChange) override;
+	virtual void PreEditChange(FEditPropertyChain& PropertyAboutToChange) override;
 	virtual void PostEditChangeProperty(FPropertyChangedEvent& PropertyChangedEvent) override;
 	virtual void PostEditChangeChainProperty(FPropertyChangedChainEvent& PropertyChangedEvent) override;
 #endif // WITH_EDITOR
