Removing Simple VirtualCamera and VirtualCamera2Actor from place actor panel. Replacing withy new VCamActor.

diff --git a/Engine/Plugins/Experimental/VirtualCamera/Source/VirtualCameraEditor/Private/VirtualCameraEditorModule.cpp b/Engine/Plugins/Experimental/VirtualCamera/Source/VirtualCameraEditor/Private/VirtualCameraEditorModule.cpp
--- a/Engine/Plugins/Experimental/VirtualCamera/Source/VirtualCameraEditor/Private/VirtualCameraEditorModule.cpp
+++ b/Engine/Plugins/Experimental/VirtualCamera/Source/VirtualCameraEditor/Private/VirtualCameraEditorModule.cpp
@@ -55,28 +55,10 @@ public:
 		{
 			if (const FPlacementCategoryInfo* Info = IVPUtilitiesEditorModule::Get().GetVirtualProductionPlacementCategoryInfo())
 			{
-				FAssetData SimpleVirtualCameraAssetData(
-					TEXT("/VirtualCamera/Deprecated/VCamCore/Blueprints/SimpleVirtualCamera"),
-					TEXT("/VirtualCamera/Deprecated/VCamCore/Blueprints"),
-					TEXT("SimpleVirtualCamera"),
-					FTopLevelAssetPath(TEXT("/Script/Engine"), TEXT("Blueprint"))
-				);
-
-				// register the simple virtual camera
-				IPlacementModeModule::Get().RegisterPlaceableItem(Info->UniqueHandle, MakeShared<FPlaceableItem>(
-					*UActorFactoryBlueprint::StaticClass(),
-					SimpleVirtualCameraAssetData,
-					FName("ClassThumbnail.CameraActor"),
-					FName("ClassIcon.CameraActor"),
-					TOptional<FLinearColor>(),
-					TOptional<int32>(),
-					NSLOCTEXT("PlacementMode", "Simple Virtual Camera", "Simple Virtual Camera")
-					));
-
 				FAssetData VirtualCamera2ActorAssetData(
-					TEXT("/VirtualCamera/Deprecated/V2/VirtualCamera2Actor"),
-					TEXT("/VirtualCamera/Deprecated/V2"),
-					TEXT("VirtualCamera2Actor"),
+					TEXT("/VirtualCamera/VCamActor"),
+					TEXT("/VirtualCamera"),
+					TEXT("VCamActor"),
 					FTopLevelAssetPath(TEXT("/Script/Engine"), TEXT("Blueprint"))
 				);
 
@@ -88,7 +70,7 @@ public:
 					FName("ClassIcon.CameraActor"),
 					TOptional<FLinearColor>(),
 					TOptional<int32>(),
-					NSLOCTEXT("PlacementMode", "VirtualCamera2 Actor", "VirtualCamera2 Actor")
+					NSLOCTEXT("PlacementMode", "VCam Actor", "VCam Actor")
 				));
 			}
 		}
