UVcamBlueprintFunctionLibrary::PilotLevelActor now also locks camera view

diff --git a/Engine/Plugins/Experimental/VirtualCamera/Source/VirtualCamera/Private/VCamBlueprintFunctionLibrary.cpp b/Engine/Plugins/Experimental/VirtualCamera/Source/VirtualCamera/Private/VCamBlueprintFunctionLibrary.cpp
--- a/Engine/Plugins/Experimental/VirtualCamera/Source/VirtualCamera/Private/VCamBlueprintFunctionLibrary.cpp
+++ b/Engine/Plugins/Experimental/VirtualCamera/Source/VirtualCamera/Private/VCamBlueprintFunctionLibrary.cpp
@@ -285,7 +285,18 @@ void UVCamBlueprintFunctionLibrary::PilotActor(AActor* SelectedActor)
 
 		if (LevelEditorSubsystem)
 		{
-			return LevelEditorSubsystem->PilotLevelActor(SelectedActor);
+			LevelEditorSubsystem->PilotLevelActor(SelectedActor);
+
+			FLevelEditorModule& LevelEditorModule = FModuleManager::GetModuleChecked<FLevelEditorModule>("LevelEditor");
+			TSharedPtr<ILevelEditor> LevelEditor = LevelEditorModule.GetFirstLevelEditor();
+			if (LevelEditor)
+			{
+				TSharedPtr<SLevelViewport> ActiveLevelViewport = LevelEditor->GetActiveViewportInterface();
+				if (ActiveLevelViewport && !ActiveLevelViewport->IsLockedCameraViewEnabled())
+				{
+					ActiveLevelViewport->ToggleActorPilotCameraView();
+				}
+			}
 		}
 	}
 #endif
