VCam plugin updated to not support Mac

diff --git a/Engine/Plugins/Experimental/VirtualCamera/VirtualCamera.uplugin b/Engine/Plugins/Experimental/VirtualCamera/VirtualCamera.uplugin
--- a/Engine/Plugins/Experimental/VirtualCamera/VirtualCamera.uplugin
+++ b/Engine/Plugins/Experimental/VirtualCamera/VirtualCamera.uplugin
@@ -14,6 +14,10 @@
 	"IsBetaVersion": true,
 	"Installed": false,
 	"EnabledByDefault": false,
+	"SupportedTargetPlatforms": [
+		"Win64",
+		"Linux"
+	],
 	"Modules": [
 		{
 			"Name": "VCamCore",
@@ -45,8 +49,7 @@
 			"LoadingPhase": "PostDefault",
 			"PlatformAllowList": [
 				"Win64",
-				"Linux",
-				"Mac"
+				"Linux"
 			]
 		},
 		{
