Add pixel streaming vcam support dependency

diff --git a/Engine/Plugins/Experimental/VirtualCamera/VirtualCamera.uplugin b/Engine/Plugins/Experimental/VirtualCamera/VirtualCamera.uplugin
--- a/Engine/Plugins/Experimental/VirtualCamera/VirtualCamera.uplugin
+++ b/Engine/Plugins/Experimental/VirtualCamera/VirtualCamera.uplugin
@@ -94,6 +94,10 @@
 		{
 			"Name": "CommonUI",
 			"Enabled": true
+		},
+		{
+			"Name": "PixelStreamingVCam",
+			"Enabled": true
 		}
 	]
 }
