Backout 22208775 to unblock UGS build issue

diff --git a/Engine/Plugins/Experimental/VirtualCamera/VirtualCamera.uplugin b/Engine/Plugins/Experimental/VirtualCamera/VirtualCamera.uplugin
--- a/Engine/Plugins/Experimental/VirtualCamera/VirtualCamera.uplugin
+++ b/Engine/Plugins/Experimental/VirtualCamera/VirtualCamera.uplugin
@@ -94,10 +94,6 @@
 		{
 			"Name": "CommonUI",
 			"Enabled": true
-		},
-		{
-			"Name": "PixelStreamingVCam",
-			"Enabled": true
 		}
 	]
 }
