Check validity of provider on session shutdown

diff --git a/Engine/Plugins/Experimental/VirtualCamera/Source/VCamCore/Private/VCamComponent.cpp b/Engine/Plugins/Experimental/VirtualCamera/Source/VCamCore/Private/VCamComponent.cpp
--- a/Engine/Plugins/Experimental/VirtualCamera/Source/VCamCore/Private/VCamComponent.cpp
+++ b/Engine/Plugins/Experimental/VirtualCamera/Source/VCamCore/Private/VCamComponent.cpp
@@ -1541,7 +1541,10 @@ void UVCamComponent::SessionShutdown(TSharedRef<IConcertClientSession> /*InSessi
 		Session->UnregisterCustomEventHandler<FMultiUserVCamCameraComponentEvent>(this);
 		for (UVCamOutputProviderBase* Provider : OutputProviders)
 		{
-			Provider->RestoreOutput();
+			if (IsValid(Provider))
+			{
+				Provider->RestoreOutput();
+			}
 		}
 	}
 
