Ensure VCam Widgets and Modifiers unregister themselves for input on deactivation or destruction

diff --git a/Engine/Plugins/Experimental/VirtualCamera/Source/VCamCore/Private/UI/VCamWidget.cpp b/Engine/Plugins/Experimental/VirtualCamera/Source/VCamCore/Private/UI/VCamWidget.cpp
--- a/Engine/Plugins/Experimental/VirtualCamera/Source/VCamCore/Private/UI/VCamWidget.cpp
+++ b/Engine/Plugins/Experimental/VirtualCamera/Source/VCamCore/Private/UI/VCamWidget.cpp
@@ -2,9 +2,9 @@
 
 #include "UI/VCamWidget.h"
 
-void UVCamWidget::BeginDestroy()
+void UVCamWidget::NativeDestruct()
 {
-	Super::BeginDestroy();
+	Super::NativeDestruct();
 	
 	if (VCamComponent)
 	{
diff --git a/Engine/Plugins/Experimental/VirtualCamera/Source/VCamCore/Private/VCamComponent.cpp b/Engine/Plugins/Experimental/VirtualCamera/Source/VCamCore/Private/VCamComponent.cpp
--- a/Engine/Plugins/Experimental/VirtualCamera/Source/VCamCore/Private/VCamComponent.cpp
+++ b/Engine/Plugins/Experimental/VirtualCamera/Source/VCamCore/Private/VCamComponent.cpp
@@ -280,6 +280,7 @@ void UVCamComponent::PreEditChange(FEditPropertyChain& PropertyAboutToChange)
 	{
 		static FName NAME_OutputProviders = GET_MEMBER_NAME_CHECKED(UVCamComponent, OutputProviders);
 		static FName NAME_ModifierStack = GET_MEMBER_NAME_CHECKED(UVCamComponent, ModifierStack);
+		static FName NAME_GeneratedModifier = GET_MEMBER_NAME_CHECKED(FModifierStackEntry, GeneratedModifier);
 		static FName NAME_Enabled = GET_MEMBER_NAME_CHECKED(UVCamComponent, bEnabled);
 
 		const FName MemberPropertyName = MemberProperty->GetFName();
@@ -335,7 +336,7 @@ void UVCamComponent::PostEditChangeProperty(FPropertyChangedEvent& PropertyChang
 		}
 		else if (PropertyName == NAME_ModifierStack)
 		{
-			EnforceModifierStackNameUniqueness();
+			ValidateModifierStack();
 		}
 		else if (PropertyName == NAME_TargetViewport)
 		{
@@ -443,6 +444,18 @@ void UVCamComponent::AddInputMappingContext(const UVCamModifier* Modifier)
 	}
 }
 
+void UVCamComponent::RemoveInputMappingContext(const UVCamModifier* Modifier)
+{
+	if (IEnhancedInputSubsystemInterface* EnhancedInputSubsystemInterface = GetEnhancedInputSubsystemInterface())
+	{
+		const UInputMappingContext* IMC = Modifier->InputMappingContext;
+		if (IsValid(IMC))
+		{
+			EnhancedInputSubsystemInterface->RemoveMappingContext(IMC);
+		}
+	}
+}
+
 void UVCamComponent::AddInputMappingContext(UInputMappingContext* Context, int32 Priority)
 {
 	if (IEnhancedInputSubsystemInterface* EnhancedInputSubsystemInterface = GetEnhancedInputSubsystemInterface())
@@ -454,6 +467,17 @@ void UVCamComponent::AddInputMappingContext(UInputMappingContext* Context, int32
 	}
 }
 
+void UVCamComponent::RemoveInputMappingContext(UInputMappingContext* Context)
+{
+	if (IEnhancedInputSubsystemInterface* EnhancedInputSubsystemInterface = GetEnhancedInputSubsystemInterface())
+	{
+		if (IsValid(Context))
+		{
+			EnhancedInputSubsystemInterface->RemoveMappingContext(Context);
+		}
+	}
+}
+
 bool UVCamComponent::SetInputProfileFromName(const FName ProfileName)
 {
 	if (const UVCamInputSettings* VCamInputSettings = GetDefault<UVCamInputSettings>())
@@ -599,21 +623,27 @@ void UVCamComponent::Update()
 
 		for (FModifierStackEntry& ModifierStackEntry : ModifierStack)
 		{
-			if (!ModifierStackEntry.bEnabled)
+			if (UVCamModifier* Modifier = ModifierStackEntry.GeneratedModifier; IsValid(Modifier))
 			{
-				continue;
-			}
+				if (ModifierStackEntry.bEnabled)
+				{
+					// Initialize the Modifier if required
+					if (Modifier->DoesRequireInitialization())
+					{
+						Modifier->Initialize(ModifierContext, InputComponent);
+						AddInputMappingContext(Modifier);
+					}
 
-			if (UVCamModifier* Modifier = ModifierStackEntry.GeneratedModifier)
-			{
-				// Initialize the Modifier if required
-				if (Modifier->DoesRequireInitialization())
+					Modifier->Apply(ModifierContext, CameraComponent, DeltaTime);
+				}
+				else
 				{
-					Modifier->Initialize(ModifierContext, InputComponent);
-					AddInputMappingContext(Modifier);
+					// If the modifier is initialized but not enabled then we deinitialize it
+					if (!Modifier->DoesRequireInitialization())
+					{
+						Modifier->Deinitialize();
+					}
 				}
-
-				Modifier->Apply(ModifierContext, CameraComponent, DeltaTime);
 			}
 		}
 
@@ -640,17 +670,27 @@ void UVCamComponent::Update()
 
 void UVCamComponent::SetEnabled(bool bNewEnabled)
 {
-	// Disable all outputs if we're no longer enabled
+	// Disable all outputs and modifiers if we're no longer enabled
 	// NOTE this must be done BEFORE setting the actual bEnabled variable because OutputProviderBase now checks the component enabled state
 	if (!bNewEnabled)
 	{
 		for (UVCamOutputProviderBase* Provider : OutputProviders)
 		{
-			if (Provider)
+			if (IsValid(Provider))
 			{
 				Provider->Deinitialize();
 			}
 		}
+
+		// There's no need to call Initialize if we are being enabled as they'll get automatically intialized
+		// the next time that the stack is evaluated
+		for (FModifierStackEntry& ModifierEntry : ModifierStack)
+		{
+			if (IsValid(ModifierEntry.GeneratedModifier))
+			{
+				ModifierEntry.GeneratedModifier->Deinitialize();
+			}
+		}
 	}
 
 	bEnabled = bNewEnabled;
@@ -661,7 +701,7 @@ void UVCamComponent::SetEnabled(bool bNewEnabled)
 	{
 		for (UVCamOutputProviderBase* Provider : OutputProviders)
 		{
-			if (Provider)
+			if (IsValid(Provider))
 			{
 				Provider->Initialize();
 			}
@@ -1008,6 +1048,7 @@ void UVCamComponent::RegisterObjectForInput(UObject* Object)
 {
 	if (IsValid(InputComponent) && IsValid(Object))
 	{
+		InputComponent->ClearBindingsForObject(Object);
 		UInputDelegateBinding::BindInputDelegates(Object->GetClass(), InputComponent, Object);
 	}
 }
@@ -1177,7 +1218,7 @@ void UVCamComponent::DestroyOutputProvider(UVCamOutputProviderBase* Provider)
 {
 	if (Provider)
 	{
-		Provider->Deinitialize();
+		// Begin Destroy will deinitialize if needed
 		Provider->ConditionalBeginDestroy();
 		Provider = nullptr;
 	}
@@ -1203,7 +1244,7 @@ void UVCamComponent::ResetAllOutputProviders()
 	}
 }
 
-void UVCamComponent::EnforceModifierStackNameUniqueness(const FString BaseName /*= "NewModifier"*/)
+void UVCamComponent::ValidateModifierStack(const FString BaseName /*= "NewModifier"*/)
 {
 	int32 ModifiedStackIndex;
 	bool bIsNewEntry;
@@ -1262,6 +1303,15 @@ void UVCamComponent::EnforceModifierStackNameUniqueness(const FString BaseName /
 				*NewModifierName.ToString(),
 				*SavedModifierStack[ModifiedStackIndex].Name.ToString());
 		}
+
+		// Check if the generated modifier was changed and if so ensure we deinitialize the old modifier
+		UVCamModifier* OldModifier = SavedModifierStack[ModifiedStackIndex].GeneratedModifier;
+		const UVCamModifier* NewModifier = ModifierStack[ModifiedStackIndex].GeneratedModifier;
+		if (OldModifier != NewModifier && IsValid(OldModifier))
+		{
+			OldModifier->Deinitialize();
+		}
+			
 	}	
 }
 
diff --git a/Engine/Plugins/Experimental/VirtualCamera/Source/VCamCore/Private/VCamModifier.cpp b/Engine/Plugins/Experimental/VirtualCamera/Source/VCamCore/Private/VCamModifier.cpp
--- a/Engine/Plugins/Experimental/VirtualCamera/Source/VCamCore/Private/VCamModifier.cpp
+++ b/Engine/Plugins/Experimental/VirtualCamera/Source/VCamCore/Private/VCamModifier.cpp
@@ -19,6 +19,19 @@ void UVCamBlueprintModifier::Initialize(UVCamModifierContext* Context, UInputCom
 	UVCamModifier::Initialize(Context, InputComponent);
 }
 
+void UVCamBlueprintModifier::Deinitialize()
+{
+	// Forward the Initialize call to the Blueprint Event only if we have a valid world
+	if (GetWorld())
+	{
+		FEditorScriptExecutionGuard ScriptGuard;
+
+		OnDeinitialize();
+	}
+	
+	Super::Deinitialize();
+}
+
 void UVCamBlueprintModifier::Apply(UVCamModifierContext* Context, UCineCameraComponent* CameraComponent, const float DeltaTime)
 {
 	// Forward the Apply call to the Blueprint Event
@@ -40,6 +53,17 @@ void UVCamModifier::Initialize(UVCamModifierContext* Context, UInputComponent* I
 	bRequiresInitialization = false;
 }
 
+void UVCamModifier::Deinitialize()
+{
+	UVCamComponent* VCamComponent = GetOwningVCamComponent();
+	if (IsValid(VCamComponent))
+	{
+		VCamComponent->UnregisterObjectForInput(this);
+	}
+
+	bRequiresInitialization = true;
+}
+
 void UVCamModifier::PostLoad()
 {
 	Super::PostLoad();
@@ -66,6 +90,11 @@ void UVCamModifier::SetEnabled(bool bNewEnabled)
 	{
 		StackEntry->bEnabled = bNewEnabled;
 	}
+
+	if (!bNewEnabled)
+	{
+		Deinitialize();
+	}
 }
 
 bool UVCamModifier::IsEnabled() const
diff --git a/Engine/Plugins/Experimental/VirtualCamera/Source/VCamCore/Public/UI/VCamWidget.h b/Engine/Plugins/Experimental/VirtualCamera/Source/VCamCore/Public/UI/VCamWidget.h
--- a/Engine/Plugins/Experimental/VirtualCamera/Source/VCamCore/Public/UI/VCamWidget.h
+++ b/Engine/Plugins/Experimental/VirtualCamera/Source/VCamCore/Public/UI/VCamWidget.h
@@ -30,7 +30,7 @@ class VCAMCORE_API UVCamWidget : public UUserWidget
 {
 	GENERATED_BODY()
 public:
-	virtual void BeginDestroy() override;
+	virtual void NativeDestruct() override;
 	
 	/*
 	 * The VCam Connections associated with this Widget
diff --git a/Engine/Plugins/Experimental/VirtualCamera/Source/VCamCore/Public/VCamComponent.h b/Engine/Plugins/Experimental/VirtualCamera/Source/VCamCore/Public/VCamComponent.h
--- a/Engine/Plugins/Experimental/VirtualCamera/Source/VCamCore/Public/VCamComponent.h
+++ b/Engine/Plugins/Experimental/VirtualCamera/Source/VCamCore/Public/VCamComponent.h
@@ -90,9 +90,15 @@ public:
 	// Adds the Input Mapping Context from a modifier, if it exists, to the input system 
 	void AddInputMappingContext(const UVCamModifier* Modifier);
 
-	// Adds and explicitly provided Input Mapping Context to the input system
+	// Removes the Input Mapping Context from a modifier, if it exists, from the input system
+	void RemoveInputMappingContext(const UVCamModifier* Modifier);
+
+	// Adds an explicitly provided Input Mapping Context to the input system
 	void AddInputMappingContext(UInputMappingContext* Context, int32 Priority);
 
+	// Removes an explicitly provided Input Mapping Context to the input system
+	void RemoveInputMappingContext(UInputMappingContext* Context);
+
 	// Attempts to apply key mapping settings from an input profile defined in VCam Input Settings
 	// Returns whether the profile was found and able to be applied
 	UFUNCTION(BlueprintCallable, Category = "VCam Input")
@@ -301,12 +307,19 @@ public:
 
 	/* Registers the given object with the VCamComponent's Input Component
 	 * This allows dynamic input bindings such as input events in blueprints to work correctly
+	 * Note: Ensure you call UnregisterObjectForInput when you are finished with the object
+	 * otherwise input events will still fire until GC actually destroys the object
 	 *
 	 * @param Object The object to register
 	 */
 	UFUNCTION(BlueprintCallable, Category="VirtualCamera")
 	void RegisterObjectForInput(UObject* Object);
 
+	/* Unregisters the given object with the VCamComponent's Input Component
+	 *
+	 * @param Object The object to unregister
+	 */
+	UFUNCTION(BlueprintCallable, Category="VirtualCamera")
 	void UnregisterObjectForInput(UObject* Object) const;
 
 	/**
@@ -339,7 +352,7 @@ private:
 
 	// Use the Saved Modifier Stack from PreEditChange to find the modified entry and then ensure the modified entry's name is unique
 	// If a new modifier has been created then its name will be defaulted to BaseName
-	void EnforceModifierStackNameUniqueness(const FString BaseName = "NewModifier");
+	void ValidateModifierStack(const FString BaseName = "NewModifier");
 	bool DoesNameExistInSavedStack(const FName InName) const;
 	void FindModifiedStackEntry(int32& ModifiedStackIndex, bool& bIsNewEntry) const;
 
diff --git a/Engine/Plugins/Experimental/VirtualCamera/Source/VCamCore/Public/VCamModifier.h b/Engine/Plugins/Experimental/VirtualCamera/Source/VCamCore/Public/VCamModifier.h
--- a/Engine/Plugins/Experimental/VirtualCamera/Source/VCamCore/Public/VCamModifier.h
+++ b/Engine/Plugins/Experimental/VirtualCamera/Source/VCamCore/Public/VCamModifier.h
@@ -33,10 +33,12 @@ class VCAMCORE_API UVCamModifier : public UObject
 
 public:
 	virtual void Initialize(UVCamModifierContext* Context, UInputComponent* InputComponent = nullptr);
+	
+	virtual void Deinitialize();
 
 	virtual void Apply(UVCamModifierContext* Context, UCineCameraComponent* CameraComponent, const float DeltaTime) {};
 
-	virtual void PostLoad();
+	virtual void PostLoad() override;
 
 	bool DoesRequireInitialization() const { return bRequiresInitialization; };
 
@@ -87,11 +89,15 @@ class VCAMCORE_API UVCamBlueprintModifier : public UVCamModifier
 
 public:
 	virtual void Initialize(UVCamModifierContext* Context, UInputComponent* InputComponent=nullptr) override;
+	virtual void Deinitialize() override;
 	virtual void Apply(UVCamModifierContext* Context, UCineCameraComponent* CameraComponent, const float DeltaTime) override;
 	
 	UFUNCTION(BlueprintImplementableEvent, Category="VirtualCamera")
 	void OnInitialize(UVCamModifierContext* Context);
 
+	UFUNCTION(BlueprintImplementableEvent, Category="VirtualCamera")
+	void OnDeinitialize();
+
 	UFUNCTION(BlueprintImplementableEvent, Category="VirtualCamera")
 	void OnApply(UVCamModifierContext* Context, UCineCameraComponent* CameraComponent, const float DeltaTime);
 
diff --git a/Engine/Plugins/Experimental/VirtualCamera/Source/VCamCore/Public/VCamTypes.h b/Engine/Plugins/Experimental/VirtualCamera/Source/VCamCore/Public/VCamTypes.h
--- a/Engine/Plugins/Experimental/VirtualCamera/Source/VCamCore/Public/VCamTypes.h
+++ b/Engine/Plugins/Experimental/VirtualCamera/Source/VCamCore/Public/VCamTypes.h
@@ -11,6 +11,7 @@
 
 class UVCamModifier;
 
+PRAGMA_DISABLE_OPTIMIZATION
 
 // Links a Modifier with a Name for use in a Modifier Stack
 USTRUCT()
@@ -56,6 +57,15 @@ struct FModifierStackEntry
 #endif
 	};
 
+	~FModifierStackEntry()
+	{
+		// It's possible the Generated Modifier may pass IsValid while having already been destroyed so ensure we explicitly check for this case
+		if (IsValid(GeneratedModifier) && !GeneratedModifier->HasAnyFlags(RF_BeginDestroyed | RF_FinishDestroyed))
+		{
+			GeneratedModifier->Deinitialize();
+		}
+	}
+
 	bool operator==(const FModifierStackEntry& Other) const
 	{
 		bool bArePropertiesEqual = this->Name.IsEqual(Other.Name) && this->bEnabled == Other.bEnabled && this->GeneratedModifier == Other.GeneratedModifier;
@@ -71,4 +81,6 @@ struct FModifierStackEntry
 	{
 		return !(*this == Other);
 	}
-};
\ No newline at end of file
+};
+
+PRAGMA_ENABLE_OPTIMIZATION
\ No newline at end of file
diff --git a/Engine/Plugins/Experimental/VirtualCamera/Source/VCamCoreEditor/Private/VCamCoreEditorModule.cpp b/Engine/Plugins/Experimental/VirtualCamera/Source/VCamCoreEditor/Private/VCamCoreEditorModule.cpp
--- a/Engine/Plugins/Experimental/VirtualCamera/Source/VCamCoreEditor/Private/VCamCoreEditorModule.cpp
+++ b/Engine/Plugins/Experimental/VirtualCamera/Source/VCamCoreEditor/Private/VCamCoreEditorModule.cpp
@@ -32,6 +32,7 @@ public:
 	void RegisterAutoGeneratedDefaultEvents()
 	{
 		FKismetEditorUtilities::RegisterAutoGeneratedDefaultEvent(this, UVCamBlueprintModifier::StaticClass(), "OnInitialize");
+		FKismetEditorUtilities::RegisterAutoGeneratedDefaultEvent(this, UVCamBlueprintModifier::StaticClass(), "OnDeinitialize");
 		FKismetEditorUtilities::RegisterAutoGeneratedDefaultEvent(this, UVCamBlueprintModifier::StaticClass(), "OnApply");
 	}	
 };
