Ensure old vcam component is properly disabled when replaced

diff --git a/Engine/Plugins/Experimental/VirtualCamera/Source/VCamCore/Private/VCamComponent.cpp b/Engine/Plugins/Experimental/VirtualCamera/Source/VCamCore/Private/VCamComponent.cpp
--- a/Engine/Plugins/Experimental/VirtualCamera/Source/VCamCore/Private/VCamComponent.cpp
+++ b/Engine/Plugins/Experimental/VirtualCamera/Source/VCamCore/Private/VCamComponent.cpp
@@ -165,7 +165,15 @@ void UVCamComponent::NotifyComponentWasReplaced(UVCamComponent* ReplacementCompo
 	ReplacementComponent->OnComponentReplaced = OnComponentReplaced;
 
 	OnComponentReplaced.Clear();
+	// Ensure all modifiers and output providers get deinitialized
+	SetEnabled(false);
 
+	// Refresh the enabled state on the new component to prevent any stale state from the replacement
+	if (ReplacementComponent->IsEnabled())
+	{
+		ReplacementComponent->SetEnabled(false);
+		ReplacementComponent->SetEnabled(true);
+	}
 	DestroyComponent();
 }
 
