Moved PixelStreamingVCam inside the Virtual Camera plugin as a module

diff --git a/Engine/Plugins/Experimental/VirtualCamera/Source/PixelStreamingVCam/PixelStreamingVCam.Build.cs b/Engine/Plugins/Experimental/VirtualCamera/Source/PixelStreamingVCam/PixelStreamingVCam.Build.cs
new file mode 100644
--- /dev/null
+++ b/Engine/Plugins/Experimental/VirtualCamera/Source/PixelStreamingVCam/PixelStreamingVCam.Build.cs
@@ -0,0 +1,46 @@
+// Copyright Epic Games, Inc. All Rights Reserved.
+
+using System;
+using System.IO;
+using System.Collections.Generic;
+
+namespace UnrealBuildTool.Rules
+{
+	public class PixelStreamingVCam : ModuleRules
+	{
+		public PixelStreamingVCam(ReadOnlyTargetRules Target) : base(Target)
+		{
+			// This is so for game projects using our public headers don't have to include extra modules they might not know about.
+			PublicDependencyModuleNames.AddRange(new string[]
+			{
+				"InputDevice",
+				"MediaIOCore",
+				"VCamCore",
+				"PixelStreaming",
+				"PixelStreamingServers"
+			});
+
+			PrivateDependencyModuleNames.AddRange(new string[]
+			{
+				"ApplicationCore",
+				"Core",
+				"CoreUObject",
+				"Engine",
+				"Projects",
+				"UMG",
+				"RHI",
+				"CinematicCamera",
+				"LiveLinkInterface",
+				"PixelCapture",
+				"PixelStreamingEditor"
+			});
+
+			// Can't package non-editor targets (e.g. games) with UnrealEd, so this dependency should only be added in editor.
+			if (Target.bBuildEditor)
+			{
+				PrivateDependencyModuleNames.Add("UnrealEd");
+			}
+
+		}
+	}
+}
diff --git a/Engine/Plugins/Experimental/VirtualCamera/Source/PixelStreamingVCam/Private/PixelStreamingMediaCapture.cpp b/Engine/Plugins/Experimental/VirtualCamera/Source/PixelStreamingVCam/Private/PixelStreamingMediaCapture.cpp
new file mode 100644
--- /dev/null
+++ b/Engine/Plugins/Experimental/VirtualCamera/Source/PixelStreamingVCam/Private/PixelStreamingMediaCapture.cpp
@@ -0,0 +1,41 @@
+// Copyright Epic Games, Inc. All Rights Reserved.
+
+#include "PixelStreamingMediaCapture.h"
+#include "PixelStreamingVideoInputRHI.h"
+
+#include "PixelCaptureInputFrameRHI.h"
+
+#include "Slate/SceneViewport.h"
+
+void UPixelStreamingMediaCapture::OnRHIResourceCaptured_RenderingThread(
+	const FCaptureBaseData& InBaseData,
+	TSharedPtr<FMediaCaptureUserData, ESPMode::ThreadSafe> InUserData,
+	FTextureRHIRef InTexture)
+{
+	if (VideoInput)
+	{
+		VideoInput->OnFrame(FPixelCaptureInputFrameRHI(InTexture));
+	}
+}
+
+bool UPixelStreamingMediaCapture::InitializeCapture()
+{
+	SetupVideoInput();
+	SetState(EMediaCaptureState::Capturing);
+	return true;
+}
+
+bool UPixelStreamingMediaCapture::PostInitializeCaptureViewport(TSharedPtr<FSceneViewport>& InSceneViewport)
+{
+	SceneViewport = TWeakPtr<FSceneViewport>(InSceneViewport);
+	OnCaptureViewportInitialized.Broadcast();
+	return true;
+}
+
+void UPixelStreamingMediaCapture::SetupVideoInput()
+{
+	if (!VideoInput)
+	{
+		VideoInput = MakeShared<FPixelStreamingVideoInputRHI>();
+	}
+}
diff --git a/Engine/Plugins/Experimental/VirtualCamera/Source/PixelStreamingVCam/Private/PixelStreamingMediaCapture.h b/Engine/Plugins/Experimental/VirtualCamera/Source/PixelStreamingVCam/Private/PixelStreamingMediaCapture.h
new file mode 100644
--- /dev/null
+++ b/Engine/Plugins/Experimental/VirtualCamera/Source/PixelStreamingVCam/Private/PixelStreamingMediaCapture.h
@@ -0,0 +1,39 @@
+// Copyright Epic Games, Inc. All Rights Reserved.
+
+#pragma once
+
+#include "MediaCapture.h"
+#include "PixelStreamingVideoInput.h"
+#include "Slate/SceneViewport.h"
+#include "PixelStreamingMediaCapture.generated.h"
+
+UCLASS(BlueprintType)
+class UPixelStreamingMediaCapture : public UMediaCapture
+{
+	GENERATED_BODY()
+
+	//~ Begin UMediaCapture interface
+public:
+	virtual void OnRHIResourceCaptured_RenderingThread(
+		const FCaptureBaseData& InBaseData,
+		TSharedPtr<FMediaCaptureUserData, ESPMode::ThreadSafe> InUserData,
+		FTextureRHIRef InTexture) override;
+
+	virtual bool InitializeCapture() override;
+	virtual bool PostInitializeCaptureViewport(TSharedPtr<FSceneViewport>& InSceneViewport) override;
+	virtual bool ShouldCaptureRHIResource() const { return true; }
+	virtual void StopCaptureImpl(bool bAllowPendingFrameToBeProcess) override {}
+	//~ End UMediaCapture interface
+
+	TSharedPtr<FPixelStreamingVideoInput> GetVideoInput() const { return VideoInput; }
+	TSharedPtr<FSceneViewport> GetViewport() const { return SceneViewport.Pin(); }
+
+	DECLARE_MULTICAST_DELEGATE(FOnCaptureViewportInitialized);
+	FOnCaptureViewportInitialized OnCaptureViewportInitialized;
+private:
+	void SetupVideoInput();
+
+private:
+	TSharedPtr<FPixelStreamingVideoInput> VideoInput;
+	TWeakPtr<FSceneViewport> SceneViewport;
+};
diff --git a/Engine/Plugins/Experimental/VirtualCamera/Source/PixelStreamingVCam/Private/PixelStreamingMediaOutput.cpp b/Engine/Plugins/Experimental/VirtualCamera/Source/PixelStreamingVCam/Private/PixelStreamingMediaOutput.cpp
new file mode 100644
--- /dev/null
+++ b/Engine/Plugins/Experimental/VirtualCamera/Source/PixelStreamingVCam/Private/PixelStreamingMediaOutput.cpp
@@ -0,0 +1,95 @@
+// Copyright Epic Games, Inc. All Rights Reserved.
+
+#include "PixelStreamingMediaOutput.h"
+#include "PixelStreamingMediaCapture.h"
+#include "IPixelStreamingModule.h"
+#include "PixelStreamingPlayerId.h"
+#include "IPixelStreamingStreamer.h"
+#include "CineCameraComponent.h"
+#include "PixelStreamingEditorModule.h"
+#include "PixelStreamingEditorUtils.h"
+
+void UPixelStreamingMediaOutput::BeginDestroy()
+{
+	StopStreaming();
+	Streamer = nullptr;
+	Super::BeginDestroy();
+}
+
+UMediaCapture* UPixelStreamingMediaOutput::CreateMediaCaptureImpl()
+{
+	if (!Streamer.IsValid())
+	{
+		IPixelStreamingModule& Module = FModuleManager::LoadModuleChecked<IPixelStreamingModule>("PixelStreaming");
+		Streamer = Module.GetStreamer(Module.GetDefaultStreamerID());
+	}
+
+	Capture = nullptr;
+	if (Streamer.IsValid())
+	{
+		Capture = NewObject<UPixelStreamingMediaCapture>();
+		Capture->SetMediaOutput(this);
+		Capture->OnStateChangedNative.AddUObject(this, &UPixelStreamingMediaOutput::OnCaptureStateChanged);
+		Capture->OnCaptureViewportInitialized.AddUObject(this, &UPixelStreamingMediaOutput::OnCaptureViewportInitialized);
+	}
+	return Capture;
+}
+
+void UPixelStreamingMediaOutput::OnCaptureStateChanged()
+{
+	switch (Capture->GetState())
+	{
+		case EMediaCaptureState::Capturing:
+			StartStreaming();
+			break;
+		case EMediaCaptureState::Stopped:
+		case EMediaCaptureState::Error:
+			StopStreaming();
+			break;
+		default:
+			break;
+	}
+}
+
+void UPixelStreamingMediaOutput::OnCaptureViewportInitialized()
+{
+	if(Streamer)
+	{
+		Streamer->SetTargetViewport(Capture->GetViewport()->GetViewportWidget());
+		Streamer->SetTargetWindow(Capture->GetViewport()->FindWindow());
+	}
+}
+
+void UPixelStreamingMediaOutput::StartStreaming()
+{
+	if (Streamer)
+	{
+		FPixelStreamingEditorModule::GetModule()->SetStreamType(UE::EditorPixelStreaming::EStreamTypes::VCam);
+		Streamer->SetVideoInput(Capture->GetVideoInput());
+
+		if (!Streamer->IsStreaming())
+		{
+			Streamer->StartStreaming();
+		}
+	}
+}
+
+void UPixelStreamingMediaOutput::StopStreaming()
+{
+	if (Streamer)
+	{
+		Streamer->StopStreaming();
+		Streamer->SetTargetViewport(nullptr);
+		Streamer->SetTargetWindow(nullptr);
+	}
+}
+
+void UPixelStreamingMediaOutput::SetSignallingServerURL(FString InURL)
+{
+	SignallingServerURL = InURL;
+}
+
+void UPixelStreamingMediaOutput::SetSignallingStreamID(FString InStreamID)
+{
+	StreamID = InStreamID;
+}
\ No newline at end of file
diff --git a/Engine/Plugins/Experimental/VirtualCamera/Source/PixelStreamingVCam/Private/PixelStreamingMediaOutput.h b/Engine/Plugins/Experimental/VirtualCamera/Source/PixelStreamingVCam/Private/PixelStreamingMediaOutput.h
new file mode 100644
--- /dev/null
+++ b/Engine/Plugins/Experimental/VirtualCamera/Source/PixelStreamingVCam/Private/PixelStreamingMediaOutput.h
@@ -0,0 +1,45 @@
+// Copyright Epic Games, Inc. All Rights Reserved.
+
+#pragma once
+
+#include "MediaOutput.h"
+#include "IPixelStreamingStreamer.h"
+#include "PixelStreamingMediaCapture.h"
+#include "PixelStreamingMediaOutput.generated.h"
+
+
+UCLASS(BlueprintType)
+class UPixelStreamingMediaOutput : public UMediaOutput
+{
+	GENERATED_BODY()
+
+	//~ Begin UMediaOutput interface
+public:
+	virtual FIntPoint GetRequestedSize() const override { return UMediaOutput::RequestCaptureSourceSize; }
+	virtual EPixelFormat GetRequestedPixelFormat() const override { return EPixelFormat::PF_B8G8R8A8; }
+	virtual EMediaCaptureConversionOperation GetConversionOperation(EMediaCaptureSourceType InSourceType) const override { return EMediaCaptureConversionOperation::NONE; }
+
+	virtual void BeginDestroy() override;
+
+protected:
+	virtual UMediaCapture* CreateMediaCaptureImpl() override;
+	//~ End UMediaOutput interface
+
+public:
+	TSharedPtr<IPixelStreamingStreamer> GetStreamer() const { return Streamer; }
+
+	void SetSignallingServerURL(FString InURL);
+	void SetSignallingStreamID(FString InStreamID);
+
+private:
+	UPixelStreamingMediaCapture* Capture = nullptr;
+	TSharedPtr<IPixelStreamingStreamer> Streamer;
+
+	FString SignallingServerURL;
+	FString StreamID = TEXT("VCam");
+
+	void OnCaptureStateChanged();
+	void OnCaptureViewportInitialized();
+	void StartStreaming();
+	void StopStreaming();
+};
diff --git a/Engine/Plugins/Experimental/VirtualCamera/Source/PixelStreamingVCam/Private/PixelStreamingVCamLog.cpp b/Engine/Plugins/Experimental/VirtualCamera/Source/PixelStreamingVCam/Private/PixelStreamingVCamLog.cpp
new file mode 100644
--- /dev/null
+++ b/Engine/Plugins/Experimental/VirtualCamera/Source/PixelStreamingVCam/Private/PixelStreamingVCamLog.cpp
@@ -0,0 +1,5 @@
+// Copyright Epic Games, Inc. All Rights Reserved.
+
+#include "PixelStreamingVCamLog.h"
+
+DEFINE_LOG_CATEGORY(LogPixelStreamingVCam);
\ No newline at end of file
diff --git a/Engine/Plugins/Experimental/VirtualCamera/Source/PixelStreamingVCam/Private/PixelStreamingVCamLog.h b/Engine/Plugins/Experimental/VirtualCamera/Source/PixelStreamingVCam/Private/PixelStreamingVCamLog.h
new file mode 100644
--- /dev/null
+++ b/Engine/Plugins/Experimental/VirtualCamera/Source/PixelStreamingVCam/Private/PixelStreamingVCamLog.h
@@ -0,0 +1,8 @@
+// Copyright Epic Games, Inc. All Rights Reserved.
+
+#pragma once
+
+#include "HAL/IConsoleManager.h"
+#include "Logging/LogMacros.h"
+
+DECLARE_LOG_CATEGORY_EXTERN(LogPixelStreamingVCam, Log, All);
\ No newline at end of file
diff --git a/Engine/Plugins/Experimental/VirtualCamera/Source/PixelStreamingVCam/Private/PixelStreamingVCamUtil.h b/Engine/Plugins/Experimental/VirtualCamera/Source/PixelStreamingVCam/Private/PixelStreamingVCamUtil.h
new file mode 100644
--- /dev/null
+++ b/Engine/Plugins/Experimental/VirtualCamera/Source/PixelStreamingVCam/Private/PixelStreamingVCamUtil.h
@@ -0,0 +1,35 @@
+// Copyright Epic Games, Inc. All Rights Reserved.
+
+#pragma once
+
+#include "Misc/Paths.h"
+#include "Interfaces/IPluginManager.h"
+#include "PixelStreamingVCamLogging.h"
+
+class FPixelStreamingVCamUtil
+{
+public:
+	static bool GetSignallingServerBinaryAbsPath(FString& OutPath)
+	{
+		TSharedPtr<IPlugin> VCamPlugin = IPluginManager::Get().FindPlugin("PixelStreamingVCam");
+
+		if(!VCamPlugin.IsValid())
+		{
+			return false;
+		}
+
+		FString PluginContentDir = FPaths::ConvertRelativePathToFull(VCamPlugin->GetContentDir());
+
+#if PLATFORM_WINDOWS
+		OutPath = PluginContentDir / TEXT("SignallingWebServer") / TEXT("cirrus.exe");
+#elif PLATFORM_LINUX
+		OutPath = PluginContentDir / TEXT("SignallingWebServer") / TEXT("cirrus");
+#else
+		UE_LOG(LogPixelStreamingVCam, Log, TEXT("Unsupported platform for Pixel Streaming VCam plugin - only Windows and Linux are supported."))
+		return false;
+#endif
+
+		// Check if the out path we set is actually valid
+		return FPaths::FileExists(OutPath);
+	}
+};
diff --git a/Engine/Plugins/Experimental/VirtualCamera/Source/PixelStreamingVCam/Private/VCamPixelStreamingLiveLink.cpp b/Engine/Plugins/Experimental/VirtualCamera/Source/PixelStreamingVCam/Private/VCamPixelStreamingLiveLink.cpp
new file mode 100644
--- /dev/null
+++ b/Engine/Plugins/Experimental/VirtualCamera/Source/PixelStreamingVCam/Private/VCamPixelStreamingLiveLink.cpp
@@ -0,0 +1,113 @@
+// Copyright Epic Games, Inc. All Rights Reserved.
+
+#include "VCamPixelStreamingLiveLink.h"
+#include "ILiveLinkClient.h"
+#include "Roles/LiveLinkTransformRole.h"
+#include "Roles/LiveLinkTransformTypes.h"
+
+#define LOCTEXT_NAMESPACE "PixelStreamingLiveLinkSource"
+
+UPixelStreamingLiveLinkSourceSettings::UPixelStreamingLiveLinkSourceSettings()
+	: ULiveLinkSourceSettings()
+{
+	// Override the default evaluation mode to latest
+	Mode = ELiveLinkSourceMode::Latest;
+}
+
+FPixelStreamingLiveLinkSource::FPixelStreamingLiveLinkSource()
+	: LiveLinkClient(nullptr)
+{
+}
+
+void FPixelStreamingLiveLinkSource::ReceiveClient(ILiveLinkClient* InClient, FGuid InSourceGuid)
+{
+	LiveLinkClient = InClient;
+	SourceGuid = InSourceGuid;
+}
+
+void FPixelStreamingLiveLinkSource::Update()
+{
+	ILiveLinkSource::Update();
+}
+
+bool FPixelStreamingLiveLinkSource::CanBeDisplayedInUI() const
+{
+	return true;
+}
+
+bool FPixelStreamingLiveLinkSource::IsSourceStillValid() const
+{
+	return true;
+}
+
+bool FPixelStreamingLiveLinkSource::RequestSourceShutdown()
+{
+	return true;
+}
+
+FText FPixelStreamingLiveLinkSource::GetSourceType() const
+{
+	return LOCTEXT("SourceType", "Pixel Streaming");
+}
+
+FText FPixelStreamingLiveLinkSource::GetSourceMachineName() const
+{
+	return FText::FromString(FPlatformProcess::ComputerName());
+}
+
+FText FPixelStreamingLiveLinkSource::GetSourceStatus() const
+{
+	return LOCTEXT("ActiveStatus", "Active");
+}
+
+void FPixelStreamingLiveLinkSource::CreateSubject(FName SubjectName) const
+{
+	if (LiveLinkClient)
+	{
+		const FLiveLinkSubjectKey SubjectKey(SourceGuid, SubjectName);
+		FLiveLinkStaticDataStruct StaticDataStruct(FLiveLinkTransformStaticData::StaticStruct());
+		LiveLinkClient->PushSubjectStaticData_AnyThread(SubjectKey, ULiveLinkTransformRole::StaticClass(), MoveTemp(StaticDataStruct));
+	}
+}
+
+void FPixelStreamingLiveLinkSource::RemoveSubject(FName SubjectName) const
+{
+	if (LiveLinkClient)
+	{
+		const FLiveLinkSubjectKey SubjectKey(SourceGuid, SubjectName);
+		LiveLinkClient->RemoveSubject_AnyThread(SubjectKey);
+	}
+}
+
+void FPixelStreamingLiveLinkSource::PushTransformForSubject(FName SubjectName, FTransform Transform) const
+{
+	if (LiveLinkClient)
+	{
+		const FLiveLinkSubjectKey SubjectKey(SourceGuid, SubjectName);
+		FLiveLinkFrameDataStruct FrameDataStruct(FLiveLinkTransformFrameData::StaticStruct());
+		FLiveLinkTransformFrameData* TransformFrameData = FrameDataStruct.Cast<FLiveLinkTransformFrameData>();
+		TransformFrameData->Transform = Transform;
+		LiveLinkClient->PushSubjectFrameData_AnyThread(SubjectKey, MoveTemp(FrameDataStruct));
+	}
+}
+
+void FPixelStreamingLiveLinkSource::PushTransformForSubject(FName SubjectName, FTransform Transform,
+	double Timestamp) const
+{
+	if (LiveLinkClient)
+	{
+		const FLiveLinkSubjectKey SubjectKey(SourceGuid, SubjectName);
+		FLiveLinkFrameDataStruct FrameDataStruct(FLiveLinkTransformFrameData::StaticStruct());
+		FLiveLinkTransformFrameData* TransformFrameData = FrameDataStruct.Cast<FLiveLinkTransformFrameData>();
+		TransformFrameData->Transform = Transform;
+
+		// Timestamp is currently assumed to be elapsed seconds at a fixed rate of 60 frames per second
+		// this will be adjusted as actual rate information is supported
+		const int32 NumberOfFrames = FMath::FloorToInt32(Timestamp * 60.0);
+		TransformFrameData->MetaData.SceneTime = FQualifiedFrameTime(NumberOfFrames, FFrameRate(60,1));
+		
+		LiveLinkClient->PushSubjectFrameData_AnyThread(SubjectKey, MoveTemp(FrameDataStruct));
+	}
+}
+
+#undef LOCTEXT_NAMESPACE
diff --git a/Engine/Plugins/Experimental/VirtualCamera/Source/PixelStreamingVCam/Private/VCamPixelStreamingLiveLink.h b/Engine/Plugins/Experimental/VirtualCamera/Source/PixelStreamingVCam/Private/VCamPixelStreamingLiveLink.h
new file mode 100644
--- /dev/null
+++ b/Engine/Plugins/Experimental/VirtualCamera/Source/PixelStreamingVCam/Private/VCamPixelStreamingLiveLink.h
@@ -0,0 +1,47 @@
+// Copyright Epic Games, Inc. All Rights Reserved.
+
+#pragma once
+
+#include "CoreMinimal.h"
+#include "ILiveLinkSource.h"
+
+#include "VCamPixelStreamingLiveLink.generated.h"
+
+UCLASS()
+class UPixelStreamingLiveLinkSourceSettings : public ULiveLinkSourceSettings
+{
+public:
+	GENERATED_BODY()
+
+	UPixelStreamingLiveLinkSourceSettings();
+};
+
+class FPixelStreamingLiveLinkSource : public ILiveLinkSource
+{
+public:
+	FPixelStreamingLiveLinkSource();
+	// ILiveLinkSource implementation
+	virtual void ReceiveClient(ILiveLinkClient* InClient, FGuid InSourceGuid) override;
+	virtual void Update() override;
+	virtual bool CanBeDisplayedInUI() const override;
+	virtual bool IsSourceStillValid() const override;
+	virtual bool RequestSourceShutdown() override;
+	virtual FText GetSourceType() const override;
+	virtual FText GetSourceMachineName() const override;
+	virtual FText GetSourceStatus() const override;
+
+	// Override the settings class to allow us to customize the default values
+	virtual TSubclassOf<ULiveLinkSourceSettings> GetSettingsClass() const override{ return UPixelStreamingLiveLinkSourceSettings::StaticClass(); };
+	
+	// Registers a new subject with the Transform Role to the Live Link Client
+	// If called with a subject name that already exists in this source then this will reset any buffered data for that subject
+	void CreateSubject(FName SubjectName) const;
+	void RemoveSubject(FName SubjectName) const;
+	void PushTransformForSubject(FName SubjectName, FTransform Transform) const;
+	void PushTransformForSubject(FName SubjectName, FTransform Transform, double Timestamp) const;
+
+private:
+	// Cached information for communicating with the live link client
+	ILiveLinkClient* LiveLinkClient;
+	FGuid SourceGuid;
+};
diff --git a/Engine/Plugins/Experimental/VirtualCamera/Source/PixelStreamingVCam/Private/VCamPixelStreamingSession.cpp b/Engine/Plugins/Experimental/VirtualCamera/Source/PixelStreamingVCam/Private/VCamPixelStreamingSession.cpp
new file mode 100644
--- /dev/null
+++ b/Engine/Plugins/Experimental/VirtualCamera/Source/PixelStreamingVCam/Private/VCamPixelStreamingSession.cpp
@@ -0,0 +1,275 @@
+// Copyright Epic Games, Inc. All Rights Reserved.
+
+#include "VCamPixelStreamingSession.h"
+#include "VCamOutputComposure.h"
+#include "PixelStreamingServers.h"
+#include "Editor/EditorPerformanceSettings.h"
+#include "Containers/UnrealString.h"
+#include "Async/Async.h"
+#include "Modules/ModuleManager.h"
+#include "Math/Vector4.h"
+#include "Math/Matrix.h"
+#include "Math/TransformNonVectorized.h"
+#include "Serialization/MemoryReader.h"
+#include "VCamPixelStreamingSubsystem.h"
+#include "VCamPixelStreamingLiveLink.h"
+#include "PixelStreamingVCamLog.h"
+#include "IPixelStreamingModule.h"
+#include "PixelStreamingProtocol.h"
+#include "PixelStreamingEditorModule.h"
+#include "VCamComponent.h"
+
+namespace VCamPixelStreamingSession
+{
+	static const FName LevelEditorName(TEXT("LevelEditor"));
+	static const FSoftClassPath EmptyUMGSoftClassPath(TEXT("/VCamCore/Assets/VCam_EmptyVisibleUMG.VCam_EmptyVisibleUMG_C"));
+} // namespace VCamPixelStreamingSession
+
+void UVCamPixelStreamingSession::Initialize()
+{
+	DisplayType = EVPWidgetDisplayType::PostProcess;
+	Super::Initialize();
+}
+
+void UVCamPixelStreamingSession::Deinitialize()
+{
+	if (MediaOutput)
+	{
+		MediaOutput->ConditionalBeginDestroy();
+	}
+	MediaOutput = nullptr;
+	Super::Deinitialize();
+}
+
+void UVCamPixelStreamingSession::Activate()
+{
+	if (UVCamPixelStreamingSubsystem* PixelStreamingSubsystem = UVCamPixelStreamingSubsystem::Get())
+	{
+		PixelStreamingSubsystem->RegisterActiveOutputProvider(this);
+		UVCamComponent* VCamComponent = GetTypedOuter<UVCamComponent>();
+		if (bAutoSetLiveLinkSubject && IsValid(VCamComponent))
+		{
+			VCamComponent->LiveLinkSubject = GetFName();
+		}
+	}
+
+	// If we don't have a UMG assigned, we still need to create an empty 'dummy' UMG in order to properly route the input back from the RemoteSession device
+	if (UMGClass == nullptr)
+	{
+		bUsingDummyUMG = true;
+		UMGClass = VCamPixelStreamingSession::EmptyUMGSoftClassPath.TryLoadClass<UUserWidget>();
+	}
+
+	if (!bInitialized)
+	{
+		UE_LOG(LogPixelStreamingVCam, Warning, TEXT("Trying to start Pixel Streaming, but has not been initialized yet"));
+		SetActive(false);
+		return;
+	}
+
+	if (MediaOutput == nullptr)
+	{
+		MediaOutput = NewObject<UPixelStreamingMediaOutput>(GetTransientPackage(), UPixelStreamingMediaOutput::StaticClass());
+	}
+
+	UEditorPerformanceSettings* Settings = GetMutableDefault<UEditorPerformanceSettings>();
+	bOldThrottleCPUWhenNotForeground = Settings->bThrottleCPUWhenNotForeground;
+	if (PreventEditorIdle)
+	{
+		Settings->bThrottleCPUWhenNotForeground = false;
+		Settings->PostEditChange();
+	}
+
+	if (StartSignallingServer)
+	{
+		if(!FPixelStreamingEditorModule::GetModule()->bUseExternalSignallingServer)
+		{
+			if (UVCamPixelStreamingSubsystem* PixelStreamingSubsystem = UVCamPixelStreamingSubsystem::Get())
+			{
+				PixelStreamingSubsystem->LaunchSignallingServer();
+			}
+		}
+		else
+		{
+			StartSignallingServer = false;
+		}
+		
+	}
+
+	FString SignallingDomain = FPixelStreamingEditorModule::GetModule()->GetSignallingDomain();
+	int32 StreamerPort = FPixelStreamingEditorModule::GetModule()->GetStreamerPort();
+	MediaOutput->SetSignallingServerURL(FString::Printf(TEXT("%s:%s"), *SignallingDomain, *FString::FromInt(StreamerPort)));
+	UE_LOG(LogPixelStreamingVCam, Log, TEXT("Activating PixelStreaming VCam Session. Endpoint: %s:%s"), *SignallingDomain, *FString::FromInt(StreamerPort));
+
+	MediaCapture = Cast<UPixelStreamingMediaCapture>(MediaOutput->CreateMediaCapture());
+
+	FMediaCaptureOptions Options;
+	Options.bResizeSourceBuffer = true;
+
+	// If we are rendering from a ComposureOutputProvider, get the requested render target and use that instead of the viewport
+	if (UVCamOutputComposure* ComposureProvider = Cast<UVCamOutputComposure>(GetOtherOutputProviderByIndex(FromComposureOutputProviderIndex)))
+	{
+		if (ComposureProvider->FinalOutputRenderTarget)
+		{
+			MediaCapture->CaptureTextureRenderTarget2D(ComposureProvider->FinalOutputRenderTarget, Options);
+			UE_LOG(LogPixelStreamingVCam, Log, TEXT("PixelStreaming set with ComposureRenderTarget"));
+		}
+		else
+		{
+			UE_LOG(LogPixelStreamingVCam, Warning, TEXT("PixelStreaming Composure usage was requested, but the specified ComposureOutputProvider has no FinalOutputRenderTarget set"));
+		}
+	}
+	else
+	{
+		TWeakPtr<FSceneViewport> SceneViewport = GetTargetSceneViewport();
+		if (TSharedPtr<FSceneViewport> PinnedSceneViewport = SceneViewport.Pin())
+		{
+			// Apply the override resolution if applicable
+			if (bUseOverrideResolution)
+			{
+				PinnedSceneViewport->SetFixedViewportSize(OverrideResolution.X, OverrideResolution.Y);
+			}
+			MediaCapture->CaptureSceneViewport(PinnedSceneViewport, Options);
+			UE_LOG(LogPixelStreamingVCam, Log, TEXT("PixelStreaming set with viewport"));
+		}
+	}
+
+	if(MediaOutput->GetStreamer())
+	{
+		IPixelStreamingModule& PixelStreamingModule = IPixelStreamingModule::Get();
+		
+		Protocol::EPixelStreamingMessageDirection MessageDirection = Protocol::EPixelStreamingMessageDirection::ToStreamer;
+
+		typedef Protocol::EPixelStreamingMessageTypes EType;
+		Protocol::FPixelStreamingInputMessage Message = Protocol::FPixelStreamingInputMessage(100, 72, {
+			// 4x4 Transform
+			EType::Float, EType::Float, EType::Float, EType::Float,
+			EType::Float, EType::Float, EType::Float, EType::Float,
+			EType::Float, EType::Float, EType::Float, EType::Float,
+			EType::Float, EType::Float, EType::Float, EType::Float,
+			// Timestamp
+			EType::Double
+		});
+
+		const TFunction<void(FMemoryReader)>& Handler = [this](FMemoryReader Ar) { 
+			// The buffer contains the transform matrix stored as 16 floats
+			FMatrix ARKitMatrix;
+			for (int32 Row = 0; Row < 4; ++Row)
+			{
+				float Col0, Col1, Col2, Col3;
+				Ar << Col0 << Col1 << Col2 << Col3;
+				ARKitMatrix.M[Row][0] = Col0;
+				ARKitMatrix.M[Row][1] = Col1;
+				ARKitMatrix.M[Row][2] = Col2;
+				ARKitMatrix.M[Row][3] = Col3;
+			}
+			ARKitMatrix.DiagnosticCheckNaN();
+
+			// Extract timestamp
+			double Timestamp;
+			Ar << Timestamp;
+
+			if(TSharedPtr<FPixelStreamingLiveLinkSource> LiveLinkSource = UVCamPixelStreamingSubsystem::Get()->TryGetLiveLinkSource(this))
+			{
+				LiveLinkSource->PushTransformForSubject(GetFName(), FTransform(ARKitMatrix), Timestamp);
+			}
+
+			/**
+			 * Code to control the level editor viewport
+			 */
+			// for (FLevelEditorViewportClient* LevelVC : GEditor->GetLevelViewportClients())
+			// {
+			// 	if (LevelVC && LevelVC->IsPerspective())
+			// 	{
+			// 		LevelVC->SetViewLocation(Translation);
+			// 		LevelVC->SetViewRotation(ModifiedRotation);
+			// 	}
+			// }
+		};
+
+		PixelStreamingModule.RegisterMessage(MessageDirection, "ARKitTransform", Message, Handler);
+	}
+	Super::Activate();
+}
+
+void UVCamPixelStreamingSession::StopSignallingServer()
+{
+	// Only stop the signalling server if we've been the ones to start it
+	if (UVCamPixelStreamingSubsystem* PixelStreamingSubsystem = UVCamPixelStreamingSubsystem::Get(); StartSignallingServer)
+	{
+		PixelStreamingSubsystem->StopSignallingServer();
+	}
+}
+
+void UVCamPixelStreamingSession::Deactivate()
+{
+	if (UVCamPixelStreamingSubsystem* PixelStreamingSubsystem = UVCamPixelStreamingSubsystem::Get())
+	{
+		PixelStreamingSubsystem->UnregisterActiveOutputProvider(this);
+	}
+
+	if (MediaCapture)
+	{
+
+		if (MediaOutput && MediaOutput->GetStreamer())
+		{
+			// Shutting streamer down before closing signalling server prevents an ugly websocket disconnect showing in the log
+			MediaOutput->GetStreamer()->StopStreaming();
+			StopSignallingServer();
+		}
+
+		MediaCapture->StopCapture(true);
+		MediaCapture = nullptr;
+	}
+	else
+	{
+		// There is not media capture we defensively clean up the signalling server if it exists.
+		StopSignallingServer();
+	}
+
+	// Remove the override resolution
+	if (bUseOverrideResolution)
+	{
+		TWeakPtr<FSceneViewport> SceneViewport = GetTargetSceneViewport();
+		if (TSharedPtr<FSceneViewport> PinnedSceneViewport = SceneViewport.Pin())
+		{
+			PinnedSceneViewport->SetFixedViewportSize(0, 0);
+		}
+	}
+
+	Super::Deactivate();
+	if (bUsingDummyUMG)
+	{
+		UMGClass = nullptr;
+		bUsingDummyUMG = false;
+	}
+
+	UEditorPerformanceSettings* Settings = GetMutableDefault<UEditorPerformanceSettings>();
+	Settings->bThrottleCPUWhenNotForeground = bOldThrottleCPUWhenNotForeground;
+	Settings->PostEditChange();
+}
+
+void UVCamPixelStreamingSession::Tick(const float DeltaTime)
+{
+	Super::Tick(DeltaTime);
+}
+
+#if WITH_EDITOR
+void UVCamPixelStreamingSession::PostEditChangeProperty(FPropertyChangedEvent& PropertyChangedEvent)
+{
+	FProperty* Property = PropertyChangedEvent.MemberProperty;
+	if (Property && PropertyChangedEvent.ChangeType != EPropertyChangeType::Interactive)
+	{
+		static FName NAME_FromComposureOutputProviderIndex = GET_MEMBER_NAME_CHECKED(UVCamPixelStreamingSession, FromComposureOutputProviderIndex);
+		static FName NAME_StartSignallingServer = GET_MEMBER_NAME_CHECKED(UVCamPixelStreamingSession, StartSignallingServer);
+		
+		if (Property->GetFName() == NAME_FromComposureOutputProviderIndex || Property->GetFName() == NAME_StartSignallingServer)
+		{
+			SetActive(false);
+		}
+	}
+	Super::PostEditChangeProperty(PropertyChangedEvent);
+}
+#endif
+
+IMPLEMENT_MODULE(FDefaultModuleImpl, PixelStreamingVCam)
diff --git a/Engine/Plugins/Experimental/VirtualCamera/Source/PixelStreamingVCam/Private/VCamPixelStreamingSubsystem.cpp b/Engine/Plugins/Experimental/VirtualCamera/Source/PixelStreamingVCam/Private/VCamPixelStreamingSubsystem.cpp
new file mode 100644
--- /dev/null
+++ b/Engine/Plugins/Experimental/VirtualCamera/Source/PixelStreamingVCam/Private/VCamPixelStreamingSubsystem.cpp
@@ -0,0 +1,98 @@
+// Copyright Epic Games, Inc. All Rights Reserved.
+
+#include "VCamPixelStreamingSubsystem.h"
+#include "ILiveLinkClient.h"
+#include "VCamPixelStreamingLiveLink.h"
+#include "VCamPixelStreamingSession.h"
+#include "PixelStreamingVCamLog.h"
+#include "PixelStreamingEditorModule.h"
+
+void UVCamPixelStreamingSubsystem::Initialize(FSubsystemCollectionBase& Collection)
+{
+	Super::Initialize(Collection);
+}
+
+void UVCamPixelStreamingSubsystem::Deinitialize()
+{
+	Super::Deinitialize();
+	IModularFeatures& ModularFeatures = IModularFeatures::Get();
+	if (LiveLinkSource && ModularFeatures.IsModularFeatureAvailable(ILiveLinkClient::ModularFeatureName))
+	{
+		ILiveLinkClient* LiveLinkClient = &ModularFeatures.GetModularFeature<ILiveLinkClient>(ILiveLinkClient::ModularFeatureName);
+		LiveLinkClient->RemoveSource(LiveLinkSource);
+	}
+	LiveLinkSource.Reset();
+}
+
+UVCamPixelStreamingSubsystem* UVCamPixelStreamingSubsystem::Get()
+{
+	return GEngine ? GEngine->GetEngineSubsystem<UVCamPixelStreamingSubsystem>() : nullptr;
+}
+
+void UVCamPixelStreamingSubsystem::RegisterActiveOutputProvider(UVCamPixelStreamingSession* OutputProvider)
+{
+	if (!OutputProvider) return;
+	
+	ActiveOutputProviders.AddUnique(OutputProvider);
+	if (LiveLinkSource)
+	{
+		LiveLinkSource->CreateSubject(OutputProvider->GetFName());
+		LiveLinkSource->PushTransformForSubject(OutputProvider->GetFName(), FTransform::Identity);
+	}
+}
+
+void UVCamPixelStreamingSubsystem::UnregisterActiveOutputProvider(UVCamPixelStreamingSession* OutputProvider)
+{
+	if (!OutputProvider) return;
+	if (ActiveOutputProviders.Remove(OutputProvider) && LiveLinkSource)
+	{
+		LiveLinkSource->RemoveSubject(OutputProvider->GetFName());
+	}
+}
+
+TSharedPtr<FPixelStreamingLiveLinkSource> UVCamPixelStreamingSubsystem::TryGetLiveLinkSource()
+{
+	if (!LiveLinkSource.IsValid())
+	{
+		IModularFeatures& ModularFeatures = IModularFeatures::Get();
+		if (ModularFeatures.IsModularFeatureAvailable(ILiveLinkClient::ModularFeatureName))
+		{
+			ILiveLinkClient* LiveLinkClient = &ModularFeatures.GetModularFeature<ILiveLinkClient>(ILiveLinkClient::ModularFeatureName);
+			LiveLinkSource = MakeShared<FPixelStreamingLiveLinkSource>();
+			LiveLinkClient->AddSource(LiveLinkSource);
+		}
+	}
+	return LiveLinkSource;
+}
+
+TSharedPtr<FPixelStreamingLiveLinkSource> UVCamPixelStreamingSubsystem::TryGetLiveLinkSource(UVCamPixelStreamingSession* OutputProvider)
+{
+	if (!LiveLinkSource.IsValid())
+	{
+		IModularFeatures& ModularFeatures = IModularFeatures::Get();
+		if (ModularFeatures.IsModularFeatureAvailable(ILiveLinkClient::ModularFeatureName))
+		{
+			ILiveLinkClient* LiveLinkClient = &ModularFeatures.GetModularFeature<ILiveLinkClient>(ILiveLinkClient::ModularFeatureName);
+			LiveLinkSource = MakeShared<FPixelStreamingLiveLinkSource>();
+			LiveLinkClient->AddSource(LiveLinkSource);
+
+			if (IsValid(OutputProvider))
+			{
+				ActiveOutputProviders.AddUnique(OutputProvider);
+				LiveLinkSource->CreateSubject(OutputProvider->GetFName());
+				LiveLinkSource->PushTransformForSubject(OutputProvider->GetFName(), FTransform::Identity);
+			}
+		}
+	}
+	return LiveLinkSource;
+}
+
+void UVCamPixelStreamingSubsystem::LaunchSignallingServer()
+{
+	FPixelStreamingEditorModule::GetModule()->StartSignalling();
+}
+
+void UVCamPixelStreamingSubsystem::StopSignallingServer()
+{
+	FPixelStreamingEditorModule::GetModule()->StopSignalling();
+}
\ No newline at end of file
diff --git a/Engine/Plugins/Experimental/VirtualCamera/Source/PixelStreamingVCam/Public/VCamPixelStreamingSession.h b/Engine/Plugins/Experimental/VirtualCamera/Source/PixelStreamingVCam/Public/VCamPixelStreamingSession.h
new file mode 100644
--- /dev/null
+++ b/Engine/Plugins/Experimental/VirtualCamera/Source/PixelStreamingVCam/Public/VCamPixelStreamingSession.h
@@ -0,0 +1,64 @@
+// Copyright Epic Games, Inc. All Rights Reserved.
+
+#pragma once
+
+#include "VCamOutputProviderBase.h"
+#include "PixelStreamingMediaCapture.h"
+#include "PixelStreamingMediaOutput.h"
+#include "PixelStreamingServers.h"
+#include "Slate/SceneViewport.h"
+#include "Engine/TextureRenderTarget2D.h"
+#if WITH_EDITOR
+	#include "LevelEditorViewport.h"
+#endif
+#include "VCamPixelStreamingSession.generated.h"
+
+UCLASS(meta = (DisplayName = "Pixel Streaming Provider (Experimental)"))
+class PIXELSTREAMINGVCAM_API UVCamPixelStreamingSession : public UVCamOutputProviderBase
+{
+	GENERATED_BODY()
+public:
+	virtual void Initialize() override;
+	virtual void Deinitialize() override;
+	virtual void Activate() override;
+	virtual void Deactivate() override;
+	virtual void Tick(const float DeltaTime) override;
+#if WITH_EDITOR
+	virtual void PostEditChangeProperty(FPropertyChangedEvent& PropertyChangedEvent) override;
+#endif	
+	// If using the output from a Composure Output Provider, specify it here
+	UPROPERTY(EditAnywhere, BlueprintReadWrite, Category = "Output")
+	int32 FromComposureOutputProviderIndex = INDEX_NONE;
+
+	// Check this if you wish to control the corresponding CineCamera with transform data received from the LiveLink app
+	UPROPERTY(EditAnywhere, BlueprintReadWrite, Category = "Output")
+	bool EnableARKitTracking = true;
+
+	// Check this if you are not separately running a signalling server for this session
+	UPROPERTY(EditAnywhere, BlueprintReadWrite, Category = "Output")
+	bool StartSignallingServer = true;
+
+	// If not selected, when the editor is not the foreground application, input through the vcam session may seem sluggish or unresponsive.
+	UPROPERTY(EditAnywhere, BlueprintReadWrite, Category = "Output")
+	bool PreventEditorIdle = true;
+
+	// If true then the Live Link Subject of the owning VCam Component will be set to the subject created by this Output Provider when the Provider is enabled
+	UPROPERTY(EditAnywhere, BlueprintReadWrite, Category = "Output")
+	bool bAutoSetLiveLinkSubject = true;
+	
+protected:
+	UPROPERTY(Transient)
+	TObjectPtr<UPixelStreamingMediaOutput> MediaOutput = nullptr;
+
+	UPROPERTY(Transient)
+	TObjectPtr<UPixelStreamingMediaCapture> MediaCapture = nullptr;
+
+private:
+	void StopSignallingServer();
+	void OnARKitTransformReceived(FPixelStreamingPlayerId PlayerId, uint8 Type, TArray<uint8> Data);
+
+private:
+	FHitResult LastViewportTouchResult;
+	bool bUsingDummyUMG = false;
+	bool bOldThrottleCPUWhenNotForeground;
+};
diff --git a/Engine/Plugins/Experimental/VirtualCamera/Source/PixelStreamingVCam/Public/VCamPixelStreamingSubsystem.h b/Engine/Plugins/Experimental/VirtualCamera/Source/PixelStreamingVCam/Public/VCamPixelStreamingSubsystem.h
new file mode 100644
--- /dev/null
+++ b/Engine/Plugins/Experimental/VirtualCamera/Source/PixelStreamingVCam/Public/VCamPixelStreamingSubsystem.h
@@ -0,0 +1,44 @@
+// Copyright Epic Games, Inc. All Rights Reserved.
+#pragma once
+
+#include "CoreMinimal.h"
+#include "Subsystems/EngineSubsystem.h"
+#include "PixelStreamingServers.h"
+#include "VCamPixelStreamingSubsystem.generated.h"
+
+class FPixelStreamingLiveLinkSource;
+class UVCamPixelStreamingSession;
+
+UCLASS()
+class PIXELSTREAMINGVCAM_API UVCamPixelStreamingSubsystem : public UEngineSubsystem
+{
+	GENERATED_BODY()
+public:
+	virtual void Initialize(FSubsystemCollectionBase& Collection) override;
+	virtual void Deinitialize() override;
+	// Convenience function for accessing the subsystem
+	static UVCamPixelStreamingSubsystem* Get();
+	void RegisterActiveOutputProvider(UVCamPixelStreamingSession* OutputProvider);
+	void UnregisterActiveOutputProvider(UVCamPixelStreamingSession* OutputProvider);
+
+	// Get the LiveLinkSource if it already exists or attempt to create one
+	TSharedPtr<FPixelStreamingLiveLinkSource> TryGetLiveLinkSource();
+	
+	// Get the LiveLinkSource if it already exists or attempt to create one
+	// Additionally registers the OutputProvider to the Source if a new source was created
+	TSharedPtr<FPixelStreamingLiveLinkSource> TryGetLiveLinkSource(UVCamPixelStreamingSession* OutputProvider);
+
+	void LaunchSignallingServer();
+	void StopSignallingServer();
+	
+private:
+	// Keep track of which output providers are currently active
+	UPROPERTY(Transient)
+	TArray<TObjectPtr<UVCamPixelStreamingSession>> ActiveOutputProviders;
+
+	// Signalling/webserver
+	TSharedPtr<UE::PixelStreamingServers::IServer> Server;
+
+	// An associated Live Link Source shared by all output providers
+	TSharedPtr<FPixelStreamingLiveLinkSource> LiveLinkSource;
+};
diff --git a/Engine/Plugins/Experimental/VirtualCamera/VirtualCamera.uplugin b/Engine/Plugins/Experimental/VirtualCamera/VirtualCamera.uplugin
--- a/Engine/Plugins/Experimental/VirtualCamera/VirtualCamera.uplugin
+++ b/Engine/Plugins/Experimental/VirtualCamera/VirtualCamera.uplugin
@@ -56,6 +56,13 @@
 			"Name": "VirtualCameraEditor",
 			"Type": "Editor",
 			"LoadingPhase": "PostEngineInit"
+		},
+		{
+			"Name": "PixelStreamingVCam",
+			"Type": "Editor",
+			"LoadingPhase": "PostEngineInit",
+			"PlatformAllowList": [ "Win64", "Linux" ],
+			"TargetDenyList": [ "Server" ]
 		}
 	],
 	"Plugins": [
@@ -94,6 +101,10 @@
 		{
 			"Name": "CommonUI",
 			"Enabled": true
+		},
+		{
+			"Name": "PixelStreaming",
+			"Enabled": true
 		}
 	]
 }
diff --git a/Engine/Plugins/Media/PixelStreamingVCam/PixelStreamingVCam.uplugin b/Engine/Plugins/Media/PixelStreamingVCam/PixelStreamingVCam.uplugin
deleted file mode 100644
diff --git a/Engine/Plugins/Media/PixelStreamingVCam/Source/PixelStreamingVCam/PixelStreamingVCam.Build.cs b/Engine/Plugins/Media/PixelStreamingVCam/Source/PixelStreamingVCam/PixelStreamingVCam.Build.cs
deleted file mode 100644
diff --git a/Engine/Plugins/Media/PixelStreamingVCam/Source/PixelStreamingVCam/Private/PixelStreamingMediaCapture.cpp b/Engine/Plugins/Media/PixelStreamingVCam/Source/PixelStreamingVCam/Private/PixelStreamingMediaCapture.cpp
deleted file mode 100644
diff --git a/Engine/Plugins/Media/PixelStreamingVCam/Source/PixelStreamingVCam/Private/PixelStreamingMediaCapture.h b/Engine/Plugins/Media/PixelStreamingVCam/Source/PixelStreamingVCam/Private/PixelStreamingMediaCapture.h
deleted file mode 100644
diff --git a/Engine/Plugins/Media/PixelStreamingVCam/Source/PixelStreamingVCam/Private/PixelStreamingMediaOutput.cpp b/Engine/Plugins/Media/PixelStreamingVCam/Source/PixelStreamingVCam/Private/PixelStreamingMediaOutput.cpp
deleted file mode 100644
diff --git a/Engine/Plugins/Media/PixelStreamingVCam/Source/PixelStreamingVCam/Private/PixelStreamingMediaOutput.h b/Engine/Plugins/Media/PixelStreamingVCam/Source/PixelStreamingVCam/Private/PixelStreamingMediaOutput.h
deleted file mode 100644
diff --git a/Engine/Plugins/Media/PixelStreamingVCam/Source/PixelStreamingVCam/Private/PixelStreamingVCamLog.cpp b/Engine/Plugins/Media/PixelStreamingVCam/Source/PixelStreamingVCam/Private/PixelStreamingVCamLog.cpp
deleted file mode 100644
diff --git a/Engine/Plugins/Media/PixelStreamingVCam/Source/PixelStreamingVCam/Private/PixelStreamingVCamLog.h b/Engine/Plugins/Media/PixelStreamingVCam/Source/PixelStreamingVCam/Private/PixelStreamingVCamLog.h
deleted file mode 100644
diff --git a/Engine/Plugins/Media/PixelStreamingVCam/Source/PixelStreamingVCam/Private/PixelStreamingVCamUtil.h b/Engine/Plugins/Media/PixelStreamingVCam/Source/PixelStreamingVCam/Private/PixelStreamingVCamUtil.h
deleted file mode 100644
diff --git a/Engine/Plugins/Media/PixelStreamingVCam/Source/PixelStreamingVCam/Private/VCamPixelStreamingLiveLink.cpp b/Engine/Plugins/Media/PixelStreamingVCam/Source/PixelStreamingVCam/Private/VCamPixelStreamingLiveLink.cpp
deleted file mode 100644
diff --git a/Engine/Plugins/Media/PixelStreamingVCam/Source/PixelStreamingVCam/Private/VCamPixelStreamingLiveLink.h b/Engine/Plugins/Media/PixelStreamingVCam/Source/PixelStreamingVCam/Private/VCamPixelStreamingLiveLink.h
deleted file mode 100644
diff --git a/Engine/Plugins/Media/PixelStreamingVCam/Source/PixelStreamingVCam/Private/VCamPixelStreamingSession.cpp b/Engine/Plugins/Media/PixelStreamingVCam/Source/PixelStreamingVCam/Private/VCamPixelStreamingSession.cpp
deleted file mode 100644
diff --git a/Engine/Plugins/Media/PixelStreamingVCam/Source/PixelStreamingVCam/Private/VCamPixelStreamingSubsystem.cpp b/Engine/Plugins/Media/PixelStreamingVCam/Source/PixelStreamingVCam/Private/VCamPixelStreamingSubsystem.cpp
deleted file mode 100644
diff --git a/Engine/Plugins/Media/PixelStreamingVCam/Source/PixelStreamingVCam/Public/VCamPixelStreamingSession.h b/Engine/Plugins/Media/PixelStreamingVCam/Source/PixelStreamingVCam/Public/VCamPixelStreamingSession.h
deleted file mode 100644
diff --git a/Engine/Plugins/Media/PixelStreamingVCam/Source/PixelStreamingVCam/Public/VCamPixelStreamingSubsystem.h b/Engine/Plugins/Media/PixelStreamingVCam/Source/PixelStreamingVCam/Public/VCamPixelStreamingSubsystem.h
deleted file mode 100644
