Fix to enable Pixel Streaming VCam to restart capturing when viewport is resized. Fix was to recreate capturer on resize and move video input onto the mediaoutput as that has a longer lifetime.

diff --git a/Engine/Plugins/Experimental/VirtualCamera/Source/PixelStreamingVCam/Private/PixelStreamingMediaCapture.cpp b/Engine/Plugins/Experimental/VirtualCamera/Source/PixelStreamingVCam/Private/PixelStreamingMediaCapture.cpp
--- a/Engine/Plugins/Experimental/VirtualCamera/Source/PixelStreamingVCam/Private/PixelStreamingMediaCapture.cpp
+++ b/Engine/Plugins/Experimental/VirtualCamera/Source/PixelStreamingVCam/Private/PixelStreamingMediaCapture.cpp
@@ -2,40 +2,47 @@
 
 #include "PixelStreamingMediaCapture.h"
 #include "PixelStreamingVideoInputRHI.h"
-
 #include "PixelCaptureInputFrameRHI.h"
-
 #include "Slate/SceneViewport.h"
+#include "PixelStreamingVCamLog.h"
 
 void UPixelStreamingMediaCapture::OnRHIResourceCaptured_RenderingThread(
 	const FCaptureBaseData& InBaseData,
 	TSharedPtr<FMediaCaptureUserData, ESPMode::ThreadSafe> InUserData,
 	FTextureRHIRef InTexture)
 {
-	if (VideoInput)
+	TSharedPtr<FPixelStreamingVideoInput> VideoInputPtr = VideoInput.Pin();
+	if (VideoInputPtr)
 	{
-		VideoInput->OnFrame(FPixelCaptureInputFrameRHI(InTexture));
+		VideoInputPtr->OnFrame(FPixelCaptureInputFrameRHI(InTexture));
 	}
 }
 
 bool UPixelStreamingMediaCapture::InitializeCapture()
 {
-	SetupVideoInput();
+	UE_LOG(LogPixelStreamingVCam, Log, TEXT("Initializing media capture for Pixel Streaming VCam."));
+	bViewportResized = false;
 	SetState(EMediaCaptureState::Capturing);
 	return true;
 }
 
+void UPixelStreamingMediaCapture::StopCaptureImpl(bool bAllowPendingFrameToBeProcess)
+{
+	// Todo: Any cleanup on capture stop should happen here.
+}
+
 bool UPixelStreamingMediaCapture::PostInitializeCaptureViewport(TSharedPtr<FSceneViewport>& InSceneViewport)
 {
 	SceneViewport = TWeakPtr<FSceneViewport>(InSceneViewport);
 	OnCaptureViewportInitialized.Broadcast();
+
+	// Listen for viewport resize events as resizes invalidate media capture, so we want to know when to reset capture
+	InSceneViewport->ViewportResizedEvent.AddUObject(this, &UPixelStreamingMediaCapture::ViewportResized);
+
 	return true;
 }
 
-void UPixelStreamingMediaCapture::SetupVideoInput()
+void UPixelStreamingMediaCapture::ViewportResized(FViewport* Viewport, uint32 ResizeCode)
 {
-	if (!VideoInput)
-	{
-		VideoInput = MakeShared<FPixelStreamingVideoInputRHI>();
-	}
+	bViewportResized = true;
 }
diff --git a/Engine/Plugins/Experimental/VirtualCamera/Source/PixelStreamingVCam/Private/PixelStreamingMediaCapture.h b/Engine/Plugins/Experimental/VirtualCamera/Source/PixelStreamingVCam/Private/PixelStreamingMediaCapture.h
--- a/Engine/Plugins/Experimental/VirtualCamera/Source/PixelStreamingVCam/Private/PixelStreamingMediaCapture.h
+++ b/Engine/Plugins/Experimental/VirtualCamera/Source/PixelStreamingVCam/Private/PixelStreamingMediaCapture.h
@@ -3,8 +3,8 @@
 #pragma once
 
 #include "MediaCapture.h"
-#include "PixelStreamingVideoInput.h"
 #include "Slate/SceneViewport.h"
+#include "PixelStreamingVideoInput.h"
 #include "PixelStreamingMediaCapture.generated.h"
 
 UCLASS(BlueprintType)
@@ -22,18 +22,21 @@ public:
 	virtual bool InitializeCapture() override;
 	virtual bool PostInitializeCaptureViewport(TSharedPtr<FSceneViewport>& InSceneViewport) override;
 	virtual bool ShouldCaptureRHIResource() const { return true; }
-	virtual void StopCaptureImpl(bool bAllowPendingFrameToBeProcess) override {}
+	virtual void StopCaptureImpl(bool bAllowPendingFrameToBeProcess) override;
 	//~ End UMediaCapture interface
 
-	TSharedPtr<FPixelStreamingVideoInput> GetVideoInput() const { return VideoInput; }
 	TSharedPtr<FSceneViewport> GetViewport() const { return SceneViewport.Pin(); }
+	virtual void ViewportResized(FViewport* Viewport, uint32 ResizeCode);
+	bool WasViewportResized() const { return bViewportResized; }
+	void SetVideoInput(TWeakPtr<FPixelStreamingVideoInput> InVideoInput) { VideoInput = InVideoInput; }
 
 	DECLARE_MULTICAST_DELEGATE(FOnCaptureViewportInitialized);
 	FOnCaptureViewportInitialized OnCaptureViewportInitialized;
-private:
-	void SetupVideoInput();
 
 private:
-	TSharedPtr<FPixelStreamingVideoInput> VideoInput;
 	TWeakPtr<FSceneViewport> SceneViewport;
+	TWeakPtr<FPixelStreamingVideoInput> VideoInput;
+	
+	/* We track whether the viewport has been resized since we created this capturer as resize means restart capturer. */
+	bool bViewportResized = false;
 };
diff --git a/Engine/Plugins/Experimental/VirtualCamera/Source/PixelStreamingVCam/Private/PixelStreamingMediaOutput.cpp b/Engine/Plugins/Experimental/VirtualCamera/Source/PixelStreamingVCam/Private/PixelStreamingMediaOutput.cpp
--- a/Engine/Plugins/Experimental/VirtualCamera/Source/PixelStreamingVCam/Private/PixelStreamingMediaOutput.cpp
+++ b/Engine/Plugins/Experimental/VirtualCamera/Source/PixelStreamingVCam/Private/PixelStreamingMediaOutput.cpp
@@ -8,6 +8,7 @@
 #include "CineCameraComponent.h"
 #include "PixelStreamingEditorModule.h"
 #include "PixelStreamingEditorUtils.h"
+#include "PixelStreamingVideoInputRHI.h"
 
 void UPixelStreamingMediaOutput::BeginDestroy()
 {
@@ -29,26 +30,17 @@ UMediaCapture* UPixelStreamingMediaOutput::CreateMediaCaptureImpl()
 	{
 		Capture = NewObject<UPixelStreamingMediaCapture>();
 		Capture->SetMediaOutput(this);
-		Capture->OnStateChangedNative.AddUObject(this, &UPixelStreamingMediaOutput::OnCaptureStateChanged);
 		Capture->OnCaptureViewportInitialized.AddUObject(this, &UPixelStreamingMediaOutput::OnCaptureViewportInitialized);
 	}
-	return Capture;
-}
 
-void UPixelStreamingMediaOutput::OnCaptureStateChanged()
-{
-	switch (Capture->GetState())
+	if (!VideoInput)
 	{
-		case EMediaCaptureState::Capturing:
-			StartStreaming();
-			break;
-		case EMediaCaptureState::Stopped:
-		case EMediaCaptureState::Error:
-			StopStreaming();
-			break;
-		default:
-			break;
+		VideoInput = MakeShared<FPixelStreamingVideoInputRHI>();
 	}
+
+	Capture->SetVideoInput(VideoInput);
+
+	return Capture;
 }
 
 void UPixelStreamingMediaOutput::OnCaptureViewportInitialized()
@@ -65,7 +57,16 @@ void UPixelStreamingMediaOutput::StartStreaming()
 	if (Streamer)
 	{
 		FPixelStreamingEditorModule::GetModule()->SetStreamType(UE::EditorPixelStreaming::EStreamTypes::VCam);
-		Streamer->SetVideoInput(Capture->GetVideoInput());
+		
+		// Only update streamer's video input if we don't have one or it is different than the one we already have.
+		if(VideoInput.IsValid())
+		{
+			TSharedPtr<FPixelStreamingVideoInput> StreamerVideoInput = Streamer->GetVideoInput().Pin();
+			if(!StreamerVideoInput.IsValid() || StreamerVideoInput != VideoInput)
+			{
+				Streamer->SetVideoInput(VideoInput);
+			}
+		}
 
 		if (!Streamer->IsStreaming())
 		{
diff --git a/Engine/Plugins/Experimental/VirtualCamera/Source/PixelStreamingVCam/Private/PixelStreamingMediaOutput.h b/Engine/Plugins/Experimental/VirtualCamera/Source/PixelStreamingVCam/Private/PixelStreamingMediaOutput.h
--- a/Engine/Plugins/Experimental/VirtualCamera/Source/PixelStreamingVCam/Private/PixelStreamingMediaOutput.h
+++ b/Engine/Plugins/Experimental/VirtualCamera/Source/PixelStreamingVCam/Private/PixelStreamingMediaOutput.h
@@ -5,6 +5,7 @@
 #include "MediaOutput.h"
 #include "IPixelStreamingStreamer.h"
 #include "PixelStreamingMediaCapture.h"
+#include "PixelStreamingVideoInput.h"
 #include "PixelStreamingMediaOutput.generated.h"
 
 
@@ -30,16 +31,17 @@ public:
 
 	void SetSignallingServerURL(FString InURL);
 	void SetSignallingStreamID(FString InStreamID);
+	void StartStreaming();
+	void StopStreaming();
 
 private:
 	UPixelStreamingMediaCapture* Capture = nullptr;
 	TSharedPtr<IPixelStreamingStreamer> Streamer;
+	TSharedPtr<FPixelStreamingVideoInput> VideoInput;
 
 	FString SignallingServerURL;
 	FString StreamID = TEXT("VCam");
 
 	void OnCaptureStateChanged();
 	void OnCaptureViewportInitialized();
-	void StartStreaming();
-	void StopStreaming();
 };
diff --git a/Engine/Plugins/Experimental/VirtualCamera/Source/PixelStreamingVCam/Private/VCamPixelStreamingSession.cpp b/Engine/Plugins/Experimental/VirtualCamera/Source/PixelStreamingVCam/Private/VCamPixelStreamingSession.cpp
--- a/Engine/Plugins/Experimental/VirtualCamera/Source/PixelStreamingVCam/Private/VCamPixelStreamingSession.cpp
+++ b/Engine/Plugins/Experimental/VirtualCamera/Source/PixelStreamingVCam/Private/VCamPixelStreamingSession.cpp
@@ -72,6 +72,9 @@
 		MediaOutput = NewObject<UPixelStreamingMediaOutput>(GetTransientPackage(), UPixelStreamingMediaOutput::StaticClass());
 	}
 
+	// We setup custom handling of ARKit transforms coming from iOS devices here
+	SetupCustomInputHandling();
+
 	UEditorPerformanceSettings* Settings = GetMutableDefault<UEditorPerformanceSettings>();
 	bOldThrottleCPUWhenNotForeground = Settings->bThrottleCPUWhenNotForeground;
 	if (PreventEditorIdle)
@@ -80,60 +83,71 @@
 		Settings->PostEditChange();
 	}
 
-	if (StartSignallingServer)
-	{
-		if(!FPixelStreamingEditorModule::GetModule()->bUseExternalSignallingServer)
-		{
-			if (UVCamPixelStreamingSubsystem* PixelStreamingSubsystem = UVCamPixelStreamingSubsystem::Get())
-			{
-				PixelStreamingSubsystem->LaunchSignallingServer();
-			}
-		}
-		else
-		{
-			StartSignallingServer = false;
-		}
-		
-	}
+	// We need signalling server to be up before we can start streaming
+	SetupSignallingServer();
 
+	// Pass signalling server info to media output, aka the streamer
 	FString SignallingDomain = FPixelStreamingEditorModule::GetModule()->GetSignallingDomain();
 	int32 StreamerPort = FPixelStreamingEditorModule::GetModule()->GetStreamerPort();
 	MediaOutput->SetSignallingServerURL(FString::Printf(TEXT("%s:%s"), *SignallingDomain, *FString::FromInt(StreamerPort)));
 	UE_LOG(LogPixelStreamingVCam, Log, TEXT("Activating PixelStreaming VCam Session. Endpoint: %s:%s"), *SignallingDomain, *FString::FromInt(StreamerPort));
 
+	SetupCapture();
+	
+	Super::Activate();
+}
+
+void UVCamPixelStreamingSession::SetupCapture()
+{
+	UE_LOG(LogPixelStreamingVCam, Log, TEXT("Create new media capture for Pixel Streaming VCam."));
+
+	if(MediaCapture)
+	{
+		MediaCapture->OnStateChangedNative.RemoveAll(this);
+	}
+
+	// Create a capturer that will capture frames from viewport and send them to streamer
 	MediaCapture = Cast<UPixelStreamingMediaCapture>(MediaOutput->CreateMediaCapture());
+	MediaCapture->OnStateChangedNative.AddUObject(this, &UVCamPixelStreamingSession::OnCaptureStateChanged);
+	StartCapture();
+}
 
-	FMediaCaptureOptions Options;
-	Options.bResizeSourceBuffer = true;
-
-	// If we are rendering from a ComposureOutputProvider, get the requested render target and use that instead of the viewport
-	if (UVCamOutputComposure* ComposureProvider = Cast<UVCamOutputComposure>(GetOtherOutputProviderByIndex(FromComposureOutputProviderIndex)))
+void UVCamPixelStreamingSession::OnCaptureStateChanged()
+{
+	if(!MediaCapture || !MediaOutput)
 	{
-		if (ComposureProvider->FinalOutputRenderTarget)
-		{
-			MediaCapture->CaptureTextureRenderTarget2D(ComposureProvider->FinalOutputRenderTarget, Options);
-			UE_LOG(LogPixelStreamingVCam, Log, TEXT("PixelStreaming set with ComposureRenderTarget"));
-		}
-		else
-		{
-			UE_LOG(LogPixelStreamingVCam, Warning, TEXT("PixelStreaming Composure usage was requested, but the specified ComposureOutputProvider has no FinalOutputRenderTarget set"));
-		}
-	}
-	else
-	{
-		TWeakPtr<FSceneViewport> SceneViewport = GetTargetSceneViewport();
-		if (TSharedPtr<FSceneViewport> PinnedSceneViewport = SceneViewport.Pin())
-		{
-			// Apply the override resolution if applicable
-			if (bUseOverrideResolution)
-			{
-				PinnedSceneViewport->SetFixedViewportSize(OverrideResolution.X, OverrideResolution.Y);
-			}
-			MediaCapture->CaptureSceneViewport(PinnedSceneViewport, Options);
-			UE_LOG(LogPixelStreamingVCam, Log, TEXT("PixelStreaming set with viewport"));
-		}
+		return;
 	}
 
+	switch (MediaCapture->GetState())
+	{
+		case EMediaCaptureState::Capturing:
+			UE_LOG(LogPixelStreamingVCam, Log, TEXT("Starting media capture and streaming for Pixel Streaming VCam."));
+			MediaOutput->StartStreaming();
+			break;
+		case EMediaCaptureState::Stopped:
+			if(MediaCapture->WasViewportResized())
+			{
+				UE_LOG(LogPixelStreamingVCam, Log, TEXT("Pixel Streaming VCam capture was stopped due to resize, going to restart capture."));
+				// If it was stopped and viewport resized we assume resize caused the stop, so try a restart of capture here.
+				SetupCapture();
+			}
+			else
+			{
+				UE_LOG(LogPixelStreamingVCam, Log, TEXT("Stopping media capture and streaming for Pixel Streaming VCam."));
+				MediaOutput->StopStreaming();
+			}
+			break;
+		case EMediaCaptureState::Error:
+			UE_LOG(LogPixelStreamingVCam, Log, TEXT("Pixel Streaming VCam capture hit an error, capturing will stop."));
+			break;
+		default:
+			break;
+	}
+}
+
+void UVCamPixelStreamingSession::SetupCustomInputHandling()
+{
 	if(MediaOutput->GetStreamer())
 	{
 		IPixelStreamingModule& PixelStreamingModule = IPixelStreamingModule::Get();
@@ -152,6 +166,12 @@
 		});
 
 		const TFunction<void(FMemoryReader)>& Handler = [this](FMemoryReader Ar) { 
+
+			if(!EnableARKitTracking)
+			{
+				return;
+			}
+
 			// The buffer contains the transform matrix stored as 16 floats
 			FMatrix ARKitMatrix;
 			for (int32 Row = 0; Row < 4; ++Row)
@@ -189,7 +209,64 @@
 
 		PixelStreamingModule.RegisterMessage(MessageDirection, "ARKitTransform", Message, Handler);
 	}
-	Super::Activate();
 }
 
+void UVCamPixelStreamingSession::StartCapture()
+{
+	if(!MediaCapture)
+	{
+		return;
+	}
+
+	FMediaCaptureOptions Options;
+	Options.bResizeSourceBuffer = true;
+
+	// If we are rendering from a ComposureOutputProvider, get the requested render target and use that instead of the viewport
+	if (UVCamOutputComposure* ComposureProvider = Cast<UVCamOutputComposure>(GetOtherOutputProviderByIndex(FromComposureOutputProviderIndex)))
+	{
+		if (ComposureProvider->FinalOutputRenderTarget)
+		{
+			MediaCapture->CaptureTextureRenderTarget2D(ComposureProvider->FinalOutputRenderTarget, Options);
+			UE_LOG(LogPixelStreamingVCam, Log, TEXT("PixelStreaming set with ComposureRenderTarget"));
+		}
+		else
+		{
+			UE_LOG(LogPixelStreamingVCam, Warning, TEXT("PixelStreaming Composure usage was requested, but the specified ComposureOutputProvider has no FinalOutputRenderTarget set"));
+		}
+	}
+	else
+	{
+		TWeakPtr<FSceneViewport> SceneViewport = GetTargetSceneViewport();
+		if (TSharedPtr<FSceneViewport> PinnedSceneViewport = SceneViewport.Pin())
+		{
+			// Apply the override resolution if applicable
+			if (bUseOverrideResolution)
+			{
+				PinnedSceneViewport->SetFixedViewportSize(OverrideResolution.X, OverrideResolution.Y);
+			}
+			MediaCapture->CaptureSceneViewport(PinnedSceneViewport, Options);
+			UE_LOG(LogPixelStreamingVCam, Log, TEXT("PixelStreaming set to capture scene viewport."));
+		}
+	}
+}
+
+void UVCamPixelStreamingSession::SetupSignallingServer()
+{
+	if (StartSignallingServer)
+	{
+		if(!FPixelStreamingEditorModule::GetModule()->bUseExternalSignallingServer)
+		{
+			if (UVCamPixelStreamingSubsystem* PixelStreamingSubsystem = UVCamPixelStreamingSubsystem::Get())
+			{
+				PixelStreamingSubsystem->LaunchSignallingServer();
+			}
+		}
+		else
+		{
+			StartSignallingServer = false;
+		}
+		
+	}
+}
+
 void UVCamPixelStreamingSession::StopSignallingServer()
diff --git a/Engine/Plugins/Experimental/VirtualCamera/Source/PixelStreamingVCam/Public/VCamPixelStreamingSession.h b/Engine/Plugins/Experimental/VirtualCamera/Source/PixelStreamingVCam/Public/VCamPixelStreamingSession.h
--- a/Engine/Plugins/Experimental/VirtualCamera/Source/PixelStreamingVCam/Public/VCamPixelStreamingSession.h
+++ b/Engine/Plugins/Experimental/VirtualCamera/Source/PixelStreamingVCam/Public/VCamPixelStreamingSession.h
@@ -13,7 +13,7 @@
 #endif
 #include "VCamPixelStreamingSession.generated.h"
 
-UCLASS(meta = (DisplayName = "Pixel Streaming Provider (Experimental)"))
+UCLASS(meta = (DisplayName = "Pixel Streaming Provider"))
 class PIXELSTREAMINGVCAM_API UVCamPixelStreamingSession : public UVCamOutputProviderBase
 {
 	GENERATED_BODY()
@@ -54,7 +54,12 @@ protected:
 	TObjectPtr<UPixelStreamingMediaCapture> MediaCapture = nullptr;
 
 private:
+	void SetupSignallingServer();
 	void StopSignallingServer();
+	void SetupCapture();
+	void StartCapture();
+	void SetupCustomInputHandling();
+	void OnCaptureStateChanged();
 	void OnARKitTransformReceived(FPixelStreamingPlayerId PlayerId, uint8 Type, TArray<uint8> Data);
 
 private:
