Fix regression: Livelink sessions and ARKit handling sometimes not registering early enough causing PixelStreamingVCam - making it completely unusable.

Fix was to ensure capturer is made before we register input handling.

diff --git a/Engine/Plugins/Experimental/VirtualCamera/Source/PixelStreamingVCam/Private/VCamPixelStreamingSession.cpp b/Engine/Plugins/Experimental/VirtualCamera/Source/PixelStreamingVCam/Private/VCamPixelStreamingSession.cpp
--- a/Engine/Plugins/Experimental/VirtualCamera/Source/PixelStreamingVCam/Private/VCamPixelStreamingSession.cpp
+++ b/Engine/Plugins/Experimental/VirtualCamera/Source/PixelStreamingVCam/Private/VCamPixelStreamingSession.cpp
@@ -43,6 +43,16 @@ void UVCamPixelStreamingSession::Deinitialize()
 
 void UVCamPixelStreamingSession::Activate()
 {
+	if (!bInitialized)
+	{
+		UE_LOG(LogPixelStreamingVCam, Warning, TEXT("Trying to start Pixel Streaming, but has not been initialized yet"));
+		SetActive(false);
+		return;
+	}
+
+	// Setup livelink source
+	UVCamPixelStreamingSubsystem::Get()->TryGetLiveLinkSource(this);
+
 	if (UVCamPixelStreamingSubsystem* PixelStreamingSubsystem = UVCamPixelStreamingSubsystem::Get())
 	{
 		PixelStreamingSubsystem->RegisterActiveOutputProvider(this);
@@ -60,21 +70,11 @@ void UVCamPixelStreamingSession::Activate()
 		UMGClass = VCamPixelStreamingSession::EmptyUMGSoftClassPath.TryLoadClass<UUserWidget>();
 	}
 
-	if (!bInitialized)
-	{
-		UE_LOG(LogPixelStreamingVCam, Warning, TEXT("Trying to start Pixel Streaming, but has not been initialized yet"));
-		SetActive(false);
-		return;
-	}
-
 	if (MediaOutput == nullptr)
 	{
 		MediaOutput = NewObject<UPixelStreamingMediaOutput>(GetTransientPackage(), UPixelStreamingMediaOutput::StaticClass());
 	}
 
-	// We setup custom handling of ARKit transforms coming from iOS devices here
-	SetupCustomInputHandling();
-
 	UEditorPerformanceSettings* Settings = GetMutableDefault<UEditorPerformanceSettings>();
 	bOldThrottleCPUWhenNotForeground = Settings->bThrottleCPUWhenNotForeground;
 	if (PreventEditorIdle)
@@ -83,6 +83,12 @@ void UVCamPixelStreamingSession::Activate()
 		Settings->PostEditChange();
 	}
 
+	// This sets up media capture and streamer
+	SetupCapture();
+
+	// We setup custom handling of ARKit transforms coming from iOS devices here
+	SetupCustomInputHandling();
+
 	// We need signalling server to be up before we can start streaming
 	SetupSignallingServer();
 
@@ -91,8 +97,6 @@ void UVCamPixelStreamingSession::Activate()
 	int32 StreamerPort = FPixelStreamingEditorModule::GetModule()->GetStreamerPort();
 	MediaOutput->SetSignallingServerURL(FString::Printf(TEXT("%s:%s"), *SignallingDomain, *FString::FromInt(StreamerPort)));
 	UE_LOG(LogPixelStreamingVCam, Log, TEXT("Activating PixelStreaming VCam Session. Endpoint: %s:%s"), *SignallingDomain, *FString::FromInt(StreamerPort));
-
-	SetupCapture();
 	
 	Super::Activate();
 }
@@ -148,7 +152,7 @@ void UVCamPixelStreamingSession::OnCaptureStateChanged()
 
 void UVCamPixelStreamingSession::SetupCustomInputHandling()
 {
-	if(MediaOutput->GetStreamer())
+	if(MediaOutput)
 	{
 		IPixelStreamingModule& PixelStreamingModule = IPixelStreamingModule::Get();
 		
@@ -209,6 +213,10 @@ void UVCamPixelStreamingSession::SetupCustomInputHandling()
 
 		PixelStreamingModule.RegisterMessage(MessageDirection, "ARKitTransform", Message, Handler);
 	}
+	else
+	{
+		UE_LOG(LogPixelStreamingVCam, Error, TEXT("Failed to setup custom input handling."));
+	}
 }
 
 void UVCamPixelStreamingSession::StartCapture()
