Fix: Editor streaming doesn't auto-launch the internal signalling server. The case where the users has not yet launched the signalling server and is not using a remote server currently leads to a bad state where pixel streaming does not work. The signalling server currently only auto-launches for VCAM whereas it should auto-launch for editor as well.

diff --git a/Engine/Plugins/Experimental/VirtualCamera/Source/PixelStreamingVCam/Private/VCamPixelStreamingSession.cpp b/Engine/Plugins/Experimental/VirtualCamera/Source/PixelStreamingVCam/Private/VCamPixelStreamingSession.cpp
--- a/Engine/Plugins/Experimental/VirtualCamera/Source/PixelStreamingVCam/Private/VCamPixelStreamingSession.cpp
+++ b/Engine/Plugins/Experimental/VirtualCamera/Source/PixelStreamingVCam/Private/VCamPixelStreamingSession.cpp
@@ -260,27 +260,17 @@ void UVCamPixelStreamingSession::StartCapture()
 
 void UVCamPixelStreamingSession::SetupSignallingServer()
 {
-	if (StartSignallingServer)
+	// Only start the signalling server if we aren't using an external signalling server
+	if (UVCamPixelStreamingSubsystem* PixelStreamingSubsystem = UVCamPixelStreamingSubsystem::Get(); !FPixelStreamingEditorModule::GetModule()->bUseExternalSignallingServer)
 	{
-		if(!FPixelStreamingEditorModule::GetModule()->bUseExternalSignallingServer)
-		{
-			if (UVCamPixelStreamingSubsystem* PixelStreamingSubsystem = UVCamPixelStreamingSubsystem::Get())
-			{
-				PixelStreamingSubsystem->LaunchSignallingServer();
-			}
-		}
-		else
-		{
-			StartSignallingServer = false;
-		}
-		
+		PixelStreamingSubsystem->LaunchSignallingServer();
 	}
 }
 
 void UVCamPixelStreamingSession::StopSignallingServer()
 {
 	// Only stop the signalling server if we've been the ones to start it
-	if (UVCamPixelStreamingSubsystem* PixelStreamingSubsystem = UVCamPixelStreamingSubsystem::Get(); StartSignallingServer)
+	if (UVCamPixelStreamingSubsystem* PixelStreamingSubsystem = UVCamPixelStreamingSubsystem::Get(); !FPixelStreamingEditorModule::GetModule()->bUseExternalSignallingServer)
 	{
 		PixelStreamingSubsystem->StopSignallingServer();
 	}
@@ -300,9 +290,9 @@ void UVCamPixelStreamingSession::Deactivate()
 		{
 			// Shutting streamer down before closing signalling server prevents an ugly websocket disconnect showing in the log
 			MediaOutput->GetStreamer()->StopStreaming();
-			StopSignallingServer();
 		}
 
+		StopSignallingServer();
 		MediaCapture->StopCapture(true);
 		MediaCapture = nullptr;
 	}
@@ -346,9 +336,8 @@ void UVCamPixelStreamingSession::PostEditChangeProperty(FPropertyChangedEvent& P
 	if (Property && PropertyChangedEvent.ChangeType != EPropertyChangeType::Interactive)
 	{
 		static FName NAME_FromComposureOutputProviderIndex = GET_MEMBER_NAME_CHECKED(UVCamPixelStreamingSession, FromComposureOutputProviderIndex);
-		static FName NAME_StartSignallingServer = GET_MEMBER_NAME_CHECKED(UVCamPixelStreamingSession, StartSignallingServer);
 		
-		if (Property->GetFName() == NAME_FromComposureOutputProviderIndex || Property->GetFName() == NAME_StartSignallingServer)
+		if (Property->GetFName() == NAME_FromComposureOutputProviderIndex)
 		{
 			SetActive(false);
 		}
diff --git a/Engine/Plugins/Experimental/VirtualCamera/Source/PixelStreamingVCam/Public/VCamPixelStreamingSession.h b/Engine/Plugins/Experimental/VirtualCamera/Source/PixelStreamingVCam/Public/VCamPixelStreamingSession.h
--- a/Engine/Plugins/Experimental/VirtualCamera/Source/PixelStreamingVCam/Public/VCamPixelStreamingSession.h
+++ b/Engine/Plugins/Experimental/VirtualCamera/Source/PixelStreamingVCam/Public/VCamPixelStreamingSession.h
@@ -34,10 +34,6 @@ public:
 	UPROPERTY(EditAnywhere, BlueprintReadWrite, Category = "Output")
 	bool EnableARKitTracking = true;
 
-	// Check this if you are not separately running a signalling server for this session
-	UPROPERTY(EditAnywhere, BlueprintReadWrite, Category = "Output")
-	bool StartSignallingServer = true;
-
 	// If not selected, when the editor is not the foreground application, input through the vcam session may seem sluggish or unresponsive.
 	UPROPERTY(EditAnywhere, BlueprintReadWrite, Category = "Output")
 	bool PreventEditorIdle = true;
diff --git a/Engine/Plugins/Media/PixelStreaming/Source/PixelStreamingEditor/Private/PixelStreamingEditorModule.cpp b/Engine/Plugins/Media/PixelStreaming/Source/PixelStreamingEditor/Private/PixelStreamingEditorModule.cpp
--- a/Engine/Plugins/Media/PixelStreaming/Source/PixelStreamingEditor/Private/PixelStreamingEditorModule.cpp
+++ b/Engine/Plugins/Media/PixelStreaming/Source/PixelStreamingEditor/Private/PixelStreamingEditorModule.cpp
@@ -187,7 +187,9 @@ void FPixelStreamingEditorModule::StartStreaming(UE::EditorPixelStreaming::EStre
 	if(!bUseExternalSignallingServer)
 	{
 		Streamer->SetSignallingServerURL(FString::Printf(TEXT("%s:%d"), *SignallingDomain, StreamerPort));
+		StartSignalling();
 	}
+
 	// Use the Pixel Streaming module's start streaming method to start all streamers
 	Module.StartStreaming();
 }
@@ -201,7 +203,13 @@ void FPixelStreamingEditorModule::StopStreaming()
 		return;
 	}
 
+	if(!bUseExternalSignallingServer)
+	{
+		StopSignalling();
+	}
+
 	Streamer->SetTargetViewport(nullptr);
+	Streamer->SetTargetWindow(nullptr);
 	// Use the Pixel Streaming module's stop streaming method to stop all streamers
 	Module.StopStreaming();
 }
