Prevent a memory leak where the VCam Pixel Streaming Subsystem would hold strong references to objects even when the object should be unloaded.

diff --git a/Engine/Plugins/Experimental/VirtualCamera/Source/PixelStreamingVCam/Public/VCamPixelStreamingSubsystem.h b/Engine/Plugins/Experimental/VirtualCamera/Source/PixelStreamingVCam/Public/VCamPixelStreamingSubsystem.h
--- a/Engine/Plugins/Experimental/VirtualCamera/Source/PixelStreamingVCam/Public/VCamPixelStreamingSubsystem.h
+++ b/Engine/Plugins/Experimental/VirtualCamera/Source/PixelStreamingVCam/Public/VCamPixelStreamingSubsystem.h
@@ -34,7 +34,7 @@ public:
 private:
 	// Keep track of which output providers are currently active
 	UPROPERTY(Transient)
-	TArray<TObjectPtr<UVCamPixelStreamingSession>> ActiveOutputProviders;
+	TArray<TWeakObjectPtr<UVCamPixelStreamingSession>> ActiveOutputProviders;
 
 	// Signalling/webserver
 	TSharedPtr<UE::PixelStreamingServers::IServer> Server;
