Ensure the VCam input component gets registered to the correct input system when intiialized

diff --git a/Engine/Plugins/Experimental/VirtualCamera/Source/VCamCore/Private/VCamComponent.cpp b/Engine/Plugins/Experimental/VirtualCamera/Source/VCamCore/Private/VCamComponent.cpp
--- a/Engine/Plugins/Experimental/VirtualCamera/Source/VCamCore/Private/VCamComponent.cpp
+++ b/Engine/Plugins/Experimental/VirtualCamera/Source/VCamCore/Private/VCamComponent.cpp
@@ -78,18 +78,6 @@ UVCamComponent::UVCamComponent()
 
 		// Setup Input
 		InputComponent = NewObject<UInputComponent>(this, UInputSettings::GetDefaultInputComponentClass(), TEXT("VCamInput0"), RF_Transient);
-		
-		
-#if WITH_EDITOR
-		if (GEditor && IsValid(InputComponent))
-		{
-			if (UEnhancedInputEditorSubsystem* Subsystem = GEditor->GetEditorSubsystem<UEnhancedInputEditorSubsystem>())
-			{
-				Subsystem->PushInputComponent(InputComponent);
-				Subsystem->StartConsumingInput();
-			}
-		}
-#endif
 
 		// Apply the Default Input profile if possible
 		if (const UVCamInputSettings* VCamInputSettings = GetDefault<UVCamInputSettings>())
@@ -113,6 +101,8 @@ void UVCamComponent::OnComponentDestroyed(bool bDestroyingHierarchy)
 		}
 	}
 
+	UnregisterInputComponent();
+
 #if WITH_EDITOR
 	// Remove all event listeners
 	if (FLevelEditorModule* LevelEditorModule = FModuleManager::GetModulePtr<FLevelEditorModule>(VCamComponent::LevelEditorName))
@@ -125,14 +115,6 @@ void UVCamComponent::OnComponentDestroyed(bool bDestroyingHierarchy)
 
 	MultiUserShutdown();
 	FCoreUObjectDelegates::OnObjectsReplaced.RemoveAll(this);
-	
-	if (GEditor && IsValid(InputComponent))
-	{
-		if (UEnhancedInputEditorSubsystem* Subsystem = GEditor->GetEditorSubsystem<UEnhancedInputEditorSubsystem>())
-		{
-			Subsystem->PopInputComponent(InputComponent);
-		}
-	}
 #endif
 }
 
@@ -196,6 +178,57 @@ IEnhancedInputSubsystemInterface* UVCamComponent::GetEnhancedInputSubsystemInter
 	return nullptr;
 }
 
+void UVCamComponent::RegisterInputComponent()
+{
+	// Ensure we start from a clean slate
+	UnregisterInputComponent();
+	
+	if (const UWorld* World = GetWorld(); IsValid(World) && World->IsGameWorld())
+	{
+		if (APlayerController* PC = World->GetFirstPlayerController())
+		{
+			PC->PushInputComponent(InputComponent);
+			bIsInputRegistered = true;
+		}
+	}
+#if WITH_EDITOR
+	else if (GEditor)
+	{
+		if (UEnhancedInputEditorSubsystem* EditorInputSubsystem = GEditor->GetEditorSubsystem<UEnhancedInputEditorSubsystem>())
+		{
+			EditorInputSubsystem->PushInputComponent(InputComponent);
+			EditorInputSubsystem->StartConsumingInput();
+			bIsInputRegistered = true;
+		}
+	}
+#endif
+}
+
+void UVCamComponent::UnregisterInputComponent()
+{
+	// Removes the component from both editor and runtime input systems if possible
+	//
+	// Note: Despite the functions being called "Pop" it's actually just removing our specific input component from
+	// the stack rather than blindly popping the top component
+	if (const UWorld* World = GetWorld(); IsValid(World) && World->IsGameWorld())
+	{
+		if (APlayerController* PC = World->GetFirstPlayerController())
+		{
+			PC->PopInputComponent(InputComponent);
+		}
+	}
+#if WITH_EDITOR
+	if (GEditor)
+	{
+		if (UEnhancedInputEditorSubsystem* EditorInputSubsystem = GEditor->GetEditorSubsystem<UEnhancedInputEditorSubsystem>())
+		{
+			EditorInputSubsystem->PopInputComponent(InputComponent);
+		}
+	}
+#endif
+	bIsInputRegistered = false;
+}
+
 bool UVCamComponent::CanUpdate() const
 {
 	UWorld* World = GetWorld();
@@ -592,6 +625,12 @@ void UVCamComponent::Update()
 		return;
 	}
 
+	// Ensure we register for input if we've not previously registered
+	if (!bIsInputRegistered)
+	{
+		RegisterInputComponent();
+	}
+
 	// If requested then disable the component if we're spawned by sequencer
 	if (bDisableComponentWhenSpawnedBySequencer)
 	{
@@ -705,15 +744,21 @@ void UVCamComponent::SetEnabled(bool bNewEnabled)
 
 	// Enable any outputs that are set to active
 	// NOTE this must be done AFTER setting the actual bEnabled variable because OutputProviderBase now checks the component enabled state
-	if (bNewEnabled && ShouldUpdateOutputProviders())
+	if (bNewEnabled)
 	{
-		for (UVCamOutputProviderBase* Provider : OutputProviders)
+		if (ShouldUpdateOutputProviders())
 		{
-			if (IsValid(Provider))
+			for (UVCamOutputProviderBase* Provider : OutputProviders)
 			{
-				Provider->Initialize();
+				if (IsValid(Provider))
+				{
+					Provider->Initialize();
+				}
 			}
 		}
+
+		// Register the input component now we're enabled
+		RegisterInputComponent();
 	}
 }
 
diff --git a/Engine/Plugins/Experimental/VirtualCamera/Source/VCamCore/Public/VCamComponent.h b/Engine/Plugins/Experimental/VirtualCamera/Source/VCamCore/Public/VCamComponent.h
--- a/Engine/Plugins/Experimental/VirtualCamera/Source/VCamCore/Public/VCamComponent.h
+++ b/Engine/Plugins/Experimental/VirtualCamera/Source/VCamCore/Public/VCamComponent.h
@@ -436,4 +436,10 @@ private:
 	// handled correctly during duplication operations (copy/paste etc)
 	UPROPERTY(Transient, DuplicateTransient)
 	TObjectPtr<UInputComponent> InputComponent;
+
+	// Utility functions for registering and unregistering our input component with the correct input system
+	void RegisterInputComponent();
+	void UnregisterInputComponent();
+
+	bool bIsInputRegistered = false;
 };
