Added logic to Pixel Streaming Output Provider's PostEditChange to handle bUseOverrideResolution

diff --git a/Engine/Plugins/Experimental/VirtualCamera/Source/PixelStreamingVCam/Private/VCamPixelStreamingSession.cpp b/Engine/Plugins/Experimental/VirtualCamera/Source/PixelStreamingVCam/Private/VCamPixelStreamingSession.cpp
--- a/Engine/Plugins/Experimental/VirtualCamera/Source/PixelStreamingVCam/Private/VCamPixelStreamingSession.cpp
+++ b/Engine/Plugins/Experimental/VirtualCamera/Source/PixelStreamingVCam/Private/VCamPixelStreamingSession.cpp
@@ -279,10 +279,8 @@ void UVCamPixelStreamingSession::StartCapture()
 		if (TSharedPtr<FSceneViewport> PinnedSceneViewport = SceneViewport.Pin())
 		{
 			// Apply the override resolution if applicable
-			if (bUseOverrideResolution)
-			{
-				PinnedSceneViewport->SetFixedViewportSize(OverrideResolution.X, OverrideResolution.Y);
-			}
+			UpdateOverrideResolution(bUseOverrideResolution, PinnedSceneViewport);
+			
 			MediaCapture->CaptureSceneViewport(PinnedSceneViewport, Options);
 			UE_LOG(LogPixelStreamingVCam, Log, TEXT("PixelStreaming set to capture scene viewport."));
 		}
@@ -336,14 +334,7 @@ void UVCamPixelStreamingSession::Deactivate()
 	}
 
 	// Remove the override resolution
-	if (bUseOverrideResolution)
-	{
-		TWeakPtr<FSceneViewport> SceneViewport = GetTargetSceneViewport();
-		if (TSharedPtr<FSceneViewport> PinnedSceneViewport = SceneViewport.Pin())
-		{
-			PinnedSceneViewport->SetFixedViewportSize(0, 0);
-		}
-	}
+	UpdateOverrideResolution(false);
 
 	Super::Deactivate();
 	if (bUsingDummyUMG)
@@ -366,14 +357,20 @@ void UVCamPixelStreamingSession::Tick(const float DeltaTime)
 void UVCamPixelStreamingSession::PostEditChangeProperty(FPropertyChangedEvent& PropertyChangedEvent)
 {
 	FProperty* Property = PropertyChangedEvent.MemberProperty;
+	const FName PropertyName = Property->GetFName();
 	if (Property && PropertyChangedEvent.ChangeType != EPropertyChangeType::Interactive)
 	{
 		static FName NAME_FromComposureOutputProviderIndex = GET_MEMBER_NAME_CHECKED(UVCamPixelStreamingSession, FromComposureOutputProviderIndex);
+		static FName NAME_UseOverrideResolution = GET_MEMBER_NAME_CHECKED(UVCamPixelStreamingSession, bUseOverrideResolution);
 		
-		if (Property->GetFName() == NAME_FromComposureOutputProviderIndex)
+		if (PropertyName == NAME_FromComposureOutputProviderIndex)
 		{
 			SetActive(false);
 		}
+		else if (PropertyName == NAME_UseOverrideResolution)
+		{
+			UpdateOverrideResolution(bUseOverrideResolution);
+		}
 	}
 	Super::PostEditChangeProperty(PropertyChangedEvent);
 }
@@ -389,4 +386,25 @@ TWeakPtr<SWindow> UVCamPixelStreamingSession::GetTargetInputWindow() const
 	return InputWindow;
 }
 
+void UVCamPixelStreamingSession::UpdateOverrideResolution(bool bApplyOverride) const
+{
+	if (const TSharedPtr<FSceneViewport> TargetSceneViewport = GetTargetSceneViewport())
+	{
+		UpdateOverrideResolution(bApplyOverride, TargetSceneViewport);
+	}
+}
+
+void UVCamPixelStreamingSession::UpdateOverrideResolution(bool bApplyOverride, const TSharedPtr<FSceneViewport>& SceneViewport) const
+{
+	// Apply the override resolution if requested otherwise remove any existing override
+	if (bApplyOverride)
+	{
+		SceneViewport->SetFixedViewportSize(OverrideResolution.X, OverrideResolution.Y);
+	}
+	else
+	{
+		SceneViewport->SetFixedViewportSize(0, 0);
+	}
+}
+
 IMPLEMENT_MODULE(FDefaultModuleImpl, PixelStreamingVCam)
diff --git a/Engine/Plugins/Experimental/VirtualCamera/Source/PixelStreamingVCam/Public/VCamPixelStreamingSession.h b/Engine/Plugins/Experimental/VirtualCamera/Source/PixelStreamingVCam/Public/VCamPixelStreamingSession.h
--- a/Engine/Plugins/Experimental/VirtualCamera/Source/PixelStreamingVCam/Public/VCamPixelStreamingSession.h
+++ b/Engine/Plugins/Experimental/VirtualCamera/Source/PixelStreamingVCam/Public/VCamPixelStreamingSession.h
@@ -58,7 +58,22 @@ private:
 	void OnCaptureStateChanged();
 	void OnARKitTransformReceived(FPixelStreamingPlayerId PlayerId, uint8 Type, TArray<uint8> Data);
 	TWeakPtr<SWindow> GetTargetInputWindow() const;
+	
+	/**
+	 * Applies or removes a fixed resolution to a viewport.
+	 * As no viewport is provided this version will try to get the Target Viewport from the VCam Component and use that as the viewport to update
+	 *
+	 * @param bApplyOverride	If true applies the override resolution settings to the viewport otherwise removes any existing override
+	 */
+	void UpdateOverrideResolution(bool bApplyOverride) const;
 
+	/**
+	 * Applies or removes a fixed resolution to a given viewport 
+	 *
+	 * @param bApplyOverride	If true applies the override resolution settings to the viewport otherwise removes any existing override
+	 * @param SceneViewport		The viewport to update the override resolution on. The function assumes this pointer has been validated already and will not test it before use.
+	 */
+	void UpdateOverrideResolution(bool bApplyOverride, const TSharedPtr<FSceneViewport>& SceneViewport) const;
 private:
 	FHitResult 	LastViewportTouchResult;
 	bool 		bUsingDummyUMG = false;
