Added a default input profile for the VCam Actor

diff --git a/Engine/Plugins/Experimental/VirtualCamera/Config/Game.ini b/Engine/Plugins/Experimental/VirtualCamera/Config/Game.ini
new file mode 100644
--- /dev/null
+++ b/Engine/Plugins/Experimental/VirtualCamera/Config/Game.ini
@@ -0,0 +1,2 @@
+[/Script/VCamInput.VCamInputSettings]
++InputProfiles=(("VCam Actor Defaults", (MappableKeyOverrides=(("Filmback", None),("NextFilmbackPreset", None),("PreviousFilmbackPreset", None),("ClipPlaneAxis", None),("NextReticle", None),("PreviousReticle", None),("ReticleOpacityAxis", None),("ReticleOpacity", None),("Reticle", None),("NextOverlay", None),("PreviousOverlay", None),("OverlayOpacityAxis", None),("OverlayOpacity", None),("Overlay", None),("CustomAspectRatio", None),("CustomAspectRatioPrecision", None),("ToggleKillRoll", None),("SetCurrentLocationAxisScaling", None),("SetAllLocationAxisScaling", None),("SelectCurrentLocationAxisScaling", None),("SetLocationStabilization", None),("SetRotationStabilization", None),("ToggleHold", None),("Crane Movement", Gamepad_RightY),("Planar Movement", Gamepad_Left2D),("SetJoystickMovementGain", None),("Pan", Gamepad_RightX),("SetJoystickRotationGain", None),("ToggleFlightMode", None),("FocalLength", None),("NextLens", None),("PreviousLens", None),("Lens", None),("Iris", None),("StopDown", None),("StopUp", None),("FocusDistance", None),("FocusDistanceAxis", None),("FocalLengthAxis", None),("ISO", None),("NextISO", None),("PreviousISO", None),("ExposureCompensation", None),("ExposureCompensationReset", None),("CreateBookmark", None),("DeleteCurrentBookmark", None),("NextBookmark", None),("PreviousBookmark", None),("CurrentBookmark", None),("NextFrame", None),("PreviousFrame", None),("SetScrubSpeed", None),("SetSequenceFrameNormalized", None),("TogglePlayPause", None),("ToPlaybackEnd", None),("ToPlaybackStart", None),("StartRecording", None),("StopRecording", None),("ToggleRecording", None),("ZeroToStage", None),("ToggleSensorCorrection", None)))))
\ No newline at end of file
