VCam workaround for FKeys being nulled when Components added via SCS are reconstructed

diff --git a/Engine/Plugins/Experimental/VirtualCamera/Source/VCamCore/Private/VCamComponent.cpp b/Engine/Plugins/Experimental/VirtualCamera/Source/VCamCore/Private/VCamComponent.cpp
--- a/Engine/Plugins/Experimental/VirtualCamera/Source/VCamCore/Private/VCamComponent.cpp
+++ b/Engine/Plugins/Experimental/VirtualCamera/Source/VCamCore/Private/VCamComponent.cpp
@@ -156,6 +156,12 @@ void UVCamComponent::NotifyComponentWasReplaced(UVCamComponent* ReplacementCompo
 		ReplacementComponent->SetEnabled(false);
 		ReplacementComponent->SetEnabled(true);
 	}
+
+	// There's a current issue where FKeys will be nulled when the component reconstructs so we'll explicitly
+	// pass the Input Profile to the new component to avoid this
+	ReplacementComponent->InputProfile = InputProfile;
+	ReplacementComponent->ApplyInputProfile();
+	
 	DestroyComponent();
 }
 
