Fixed UVCamComponent::AddInputProfileWithCurrentlyActiveMappings so that it correctly handles mappings with no default key assigned

diff --git a/Engine/Plugins/Experimental/VirtualCamera/Source/VCamCore/Private/VCamComponent.cpp b/Engine/Plugins/Experimental/VirtualCamera/Source/VCamCore/Private/VCamComponent.cpp
--- a/Engine/Plugins/Experimental/VirtualCamera/Source/VCamCore/Private/VCamComponent.cpp
+++ b/Engine/Plugins/Experimental/VirtualCamera/Source/VCamCore/Private/VCamComponent.cpp
@@ -232,6 +232,7 @@ void UVCamComponent::UnregisterInputComponent()
 		}
 	}
 #endif
+	AppliedInputContexts.Reset();
 	bIsInputRegistered = false;
 }
 
@@ -486,6 +487,7 @@ void UVCamComponent::AddInputMappingContext(const UVCamModifier* Modifier)
 			if (!EnhancedInputSubsystemInterface->HasMappingContext(IMC))
 			{
 				EnhancedInputSubsystemInterface->AddMappingContext(IMC, InputPriority);
+				AppliedInputContexts.AddUnique(IMC);
 			}
 		}
 	}
@@ -499,6 +501,7 @@ void UVCamComponent::RemoveInputMappingContext(const UVCamModifier* Modifier)
 		if (IsValid(IMC))
 		{
 			EnhancedInputSubsystemInterface->RemoveMappingContext(IMC);
+			AppliedInputContexts.Remove(IMC);
 		}
 	}
 }
@@ -510,6 +513,7 @@ void UVCamComponent::AddInputMappingContext(UInputMappingContext* Context, int32
 		if (IsValid(Context))
 		{
 			EnhancedInputSubsystemInterface->AddMappingContext(Context, Priority);
+			AppliedInputContexts.AddUnique(Context);
 		}
 	}
 }
@@ -521,6 +525,7 @@ void UVCamComponent::RemoveInputMappingContext(UInputMappingContext* Context)
 		if (IsValid(Context))
 		{
 			EnhancedInputSubsystemInterface->RemoveMappingContext(Context);
+			AppliedInputContexts.Remove(Context);
 		}
 	}
 }
@@ -565,21 +570,31 @@ bool UVCamComponent::AddInputProfileWithCurrentlyActiveMappings(const FName Prof
 
 	if (const IEnhancedInputSubsystemInterface* EnhancedInputSubsystemInterface = GetEnhancedInputSubsystemInterface())
 	{
-		const TArray<FEnhancedActionKeyMapping> PlayerMappableActionKeyMappings = EnhancedInputSubsystemInterface->GetAllPlayerMappableActionKeyMappings();
-		for (const FEnhancedActionKeyMapping& PlayerMappableActionKeyMapping : PlayerMappableActionKeyMappings)
+		TArray<FEnhancedActionKeyMapping> PlayerMappableActionKeyMappings;
+		for (const UInputMappingContext* MappingContext : AppliedInputContexts)
 		{
-			const FName MappingName = PlayerMappableActionKeyMapping.PlayerMappableOptions.Name;
-
-			// Prefer to use the current mapped key but fallback to the default if no key is mapped
-			FKey CurrentKey = EnhancedInputSubsystemInterface->GetPlayerMappedKey(MappingName);
-			if (!CurrentKey.IsValid())
+			if (!IsValid(MappingContext))
 			{
-				CurrentKey = PlayerMappableActionKeyMapping.Key;
+				continue;
 			}
-
-			if (!TargetInputProfile->MappableKeyOverrides.Contains(MappingName))
+			for (const FEnhancedActionKeyMapping& Mapping : MappingContext->GetMappings())
 			{
-				TargetInputProfile->MappableKeyOverrides.Add(MappingName, CurrentKey);
+				if (Mapping.bIsPlayerMappable)
+				{
+					const FName MappingName = Mapping.PlayerMappableOptions.Name;
+
+					// Prefer to use the current mapped key but fallback to the default if no key is mapped
+					FKey CurrentKey = EnhancedInputSubsystemInterface->GetPlayerMappedKey(MappingName);
+					if (!CurrentKey.IsValid())
+					{
+						CurrentKey = Mapping.Key;
+					}
+
+					if (!TargetInputProfile->MappableKeyOverrides.Contains(MappingName))
+					{
+						TargetInputProfile->MappableKeyOverrides.Add(MappingName, CurrentKey);
+					}
+				}
 			}
 		}
 	}
diff --git a/Engine/Plugins/Experimental/VirtualCamera/Source/VCamCore/Public/VCamComponent.h b/Engine/Plugins/Experimental/VirtualCamera/Source/VCamCore/Public/VCamComponent.h
--- a/Engine/Plugins/Experimental/VirtualCamera/Source/VCamCore/Public/VCamComponent.h
+++ b/Engine/Plugins/Experimental/VirtualCamera/Source/VCamCore/Public/VCamComponent.h
@@ -441,5 +441,9 @@ private:
 	void RegisterInputComponent();
 	void UnregisterInputComponent();
 
+	// Store the Input Mapping Contexts that have been added via this component
+	UPROPERTY(Transient, DuplicateTransient)
+	TArray<TObjectPtr<const UInputMappingContext>> AppliedInputContexts;
+
 	bool bIsInputRegistered = false;
 };
