Vcam: Prevent VCam Components from activating output providers or updating in asset editor worlds

diff --git a/Engine/Plugins/Experimental/VirtualCamera/Source/VCamCore/Private/VCamComponent.cpp b/Engine/Plugins/Experimental/VirtualCamera/Source/VCamCore/Private/VCamComponent.cpp
--- a/Engine/Plugins/Experimental/VirtualCamera/Source/VCamCore/Private/VCamComponent.cpp
+++ b/Engine/Plugins/Experimental/VirtualCamera/Source/VCamCore/Private/VCamComponent.cpp
@@ -241,10 +241,10 @@ bool UVCamComponent::CanUpdate() const
 	UWorld* World = GetWorld();
 	if (bEnabled && IsValid(this) && !bIsEditorObjectButPIEIsRunning && World)
 	{
-		// Check for an Inactive type of world which means nothing should ever execute on this object
-		// @TODO: This is far from optimal as it means a zombie object has been created that never gets GC'ed
-		// Apparently, we should be using OnRegister/OnUnregister() instead of doing everything in the constructor, but it was throwing GC errors when trying that
-		if (World->WorldType != EWorldType::Inactive)
+		// Check that we only update in a valid world type
+		// This prevents us updating in asset editors or invalid worlds
+		constexpr int ValidWorldTypes = EWorldType::Game | EWorldType::PIE | EWorldType::Editor;
+		if ((World->WorldType & ValidWorldTypes) != EWorldType::None)
 		{
 			if (const USceneComponent* ParentComponent = GetAttachParent())
 			{
@@ -767,7 +767,7 @@ void UVCamComponent::SetEnabled(bool bNewEnabled)
 	// NOTE this must be done AFTER setting the actual bEnabled variable because OutputProviderBase now checks the component enabled state
 	if (bNewEnabled)
 	{
-		if (ShouldUpdateOutputProviders())
+		if (ShouldUpdateOutputProviders() && CanUpdate())
 		{
 			for (UVCamOutputProviderBase* Provider : OutputProviders)
 			{
