Fix NonUnity errors with UI Classes

diff --git a/Engine/Plugins/Experimental/CommonUI/Source/CommonUI/Private/CommonRichTextBlock.cpp b/Engine/Plugins/Experimental/CommonUI/Source/CommonUI/Private/CommonRichTextBlock.cpp
--- a/Engine/Plugins/Experimental/CommonUI/Source/CommonUI/Private/CommonRichTextBlock.cpp
+++ b/Engine/Plugins/Experimental/CommonUI/Source/CommonUI/Private/CommonRichTextBlock.cpp
@@ -24,6 +24,7 @@
 #include "Types/ReflectionMetadata.h"
 #include "Framework/Text/IRichTextMarkupWriter.h"
 #include "Framework/Text/RichTextMarkupProcessing.h"
+#include "Framework/Application/SlateApplication.h"
 
 #include UE_INLINE_GENERATED_CPP_BY_NAME(CommonRichTextBlock)
 
diff --git a/Engine/Plugins/Experimental/CommonUI/Source/CommonUI/Private/Input/CommonBoundActionButton.cpp b/Engine/Plugins/Experimental/CommonUI/Source/CommonUI/Private/Input/CommonBoundActionButton.cpp
--- a/Engine/Plugins/Experimental/CommonUI/Source/CommonUI/Private/Input/CommonBoundActionButton.cpp
+++ b/Engine/Plugins/Experimental/CommonUI/Source/CommonUI/Private/Input/CommonBoundActionButton.cpp
@@ -4,6 +4,7 @@
 #include "CommonTextBlock.h"
 #include "Input/UIActionBinding.h"
 #include "CommonActionWidget.h"
+#include "Framework/Application/SlateApplication.h"
 
 #include UE_INLINE_GENERATED_CPP_BY_NAME(CommonBoundActionButton)
 
diff --git a/Engine/Plugins/Experimental/VirtualCamera/Source/VCamInput/Private/VCamAnyKeyWidget.cpp b/Engine/Plugins/Experimental/VirtualCamera/Source/VCamInput/Private/VCamAnyKeyWidget.cpp
--- a/Engine/Plugins/Experimental/VirtualCamera/Source/VCamInput/Private/VCamAnyKeyWidget.cpp
+++ b/Engine/Plugins/Experimental/VirtualCamera/Source/VCamInput/Private/VCamAnyKeyWidget.cpp
@@ -1,7 +1,9 @@
 // Copyright Epic Games, Inc. All Rights Reserved.
 
 #include "VCamAnyKeyWidget.h"
+#include "Containers/Ticker.h"
 #include "Framework/Application/IInputProcessor.h"
+#include "Framework/Application/SlateApplication.h"
 
 // Input processor class that forwards any relevant input to the owning widget
 // All Handle events return true to consume the input
