Corrected scope issue when registering Input Mapping Contexts in the VCam Component

diff --git a/Engine/Plugins/Experimental/VirtualCamera/Source/VCamCore/Private/VCamComponent.cpp b/Engine/Plugins/Experimental/VirtualCamera/Source/VCamCore/Private/VCamComponent.cpp
--- a/Engine/Plugins/Experimental/VirtualCamera/Source/VCamCore/Private/VCamComponent.cpp
+++ b/Engine/Plugins/Experimental/VirtualCamera/Source/VCamCore/Private/VCamComponent.cpp
@@ -487,8 +487,9 @@ void UVCamComponent::AddInputMappingContext(const UVCamModifier* Modifier)
 			if (!EnhancedInputSubsystemInterface->HasMappingContext(IMC))
 			{
 				EnhancedInputSubsystemInterface->AddMappingContext(IMC, InputPriority);
-				AppliedInputContexts.AddUnique(IMC);
 			}
+			// Ensure we store the IMC even if it's already in the system as it could have been registered from outside the VCam 
+			AppliedInputContexts.AddUnique(IMC);
 		}
 	}
 }
