VCam no longer supports mac for 5.1

diff --git a/Engine/Plugins/Experimental/VirtualCamera/VirtualCamera.uplugin b/Engine/Plugins/Experimental/VirtualCamera/VirtualCamera.uplugin
--- a/Engine/Plugins/Experimental/VirtualCamera/VirtualCamera.uplugin
+++ b/Engine/Plugins/Experimental/VirtualCamera/VirtualCamera.uplugin
@@ -18,6 +18,10 @@
 		"Win64",
 		"Linux"
 	],
+	"PlatformAllowList": [
+		"Win64",
+		"Linux"
+	],
 	"Modules": [
 		{
 			"Name": "VCamCore",
