Corrected Virtual Camera config files

diff --git a/Engine/Plugins/Experimental/VirtualCamera/Config/BaseEngine.ini b/Engine/Plugins/Experimental/VirtualCamera/Config/BaseEngine.ini
deleted file mode 100644
diff --git a/Engine/Plugins/Experimental/VirtualCamera/Config/Engine.ini b/Engine/Plugins/Experimental/VirtualCamera/Config/Engine.ini
new file mode 100644
--- /dev/null
+++ b/Engine/Plugins/Experimental/VirtualCamera/Config/Engine.ini
@@ -0,0 +1,79 @@
+[CoreRedirects]
++ClassRedirects=(OldName="/Script/VCamCore.VCamCoreSubsystem",NewName="/Script/VCamInput.VCamInputSubsystem")
++ClassRedirects=(OldName="/Script/VirtualCamera.RadialSlider",NewName="/Script/AdvancedWidgets.RadialSlider")
++FunctionRedirects=(OldName="VCamComponent.ClearModifiers",NewName="VCamComponent.RemoveAllModifiers")
++FunctionRedirects=(OldName="VCamComponent.FindModifierByName",NewName="VCamComponent.GetModifierByName")
++FunctionRedirects=(OldName="VCamComponent.FindModifiersByClass",NewName="VCamComponent.GetModifiersByClass")
++FunctionRedirects=(OldName="VCamComponent.FindModifierByInterface",NewName="VCamComponent.GetModifiersByInterface")
++PackageRedirects=(OldName="/VirtualCamera/Editor", NewName="/VirtualCamera/Deprecated/Editor", MatchSubstring=true)
++PackageRedirects=(OldName="/VirtualCamera/OperatorPanel", NewName="/VirtualCamera/Deprecated/OperatorPanel", MatchSubstring=true)
++PackageRedirects=(OldName="/VirtualCamera/TestMaps", NewName="/VirtualCamera/Deprecated/TestMaps", MatchSubstring=true)
++PackageRedirects=(OldName="/VirtualCamera/V2", NewName="/VirtualCamera/Deprecated/V2", MatchSubstring=true)
++PackageRedirects=(OldName="/VirtualCamera/VCamCore", NewName="/VirtualCamera/Deprecated/VCamCore", MatchSubstring=true)
+
+[/Script/CinematicCamera.CineCameraSettings]
+DefaultLensPresetName=18mm
+DefaultLensFocalLength=18.000000
+DefaultLensFStop=5.600000
+-LensPresets=(Name="12mm Prime f/2.8",LensSettings=(MinFocalLength=12,MaxFocalLength=12,MinFStop=2.8,MaxFStop=22,DiaphragmBladeCount=7))
+-LensPresets=(Name="30mm Prime f/1.4",LensSettings=(MinFocalLength=30,MaxFocalLength=30,MinFStop=1.4,MaxFStop=22,DiaphragmBladeCount=7))
+-LensPresets=(Name="50mm Prime f/1.8",LensSettings=(MinFocalLength=50,MaxFocalLength=50,MinFStop=1.8,MaxFStop=22,DiaphragmBladeCount=7))
+-LensPresets=(Name="85mm Prime f/1.8",LensSettings=(MinFocalLength=85,MaxFocalLength=85,MinFStop=1.8,MaxFStop=22,DiaphragmBladeCount=7))
+-LensPresets=(Name="105mm Prime f/2",LensSettings=(MinFocalLength=105,MaxFocalLength=105,MinFStop=2,MaxFStop=22,DiaphragmBladeCount=7))
+-LensPresets=(Name="200mm Prime f/2",LensSettings=(MinFocalLength=200,MaxFocalLength=200,MinFStop=2,MaxFStop=22,DiaphragmBladeCount=7))
+-LensPresets=(Name="24-70mm Zoom f/2.8",LensSettings=(MinFocalLength=24,MaxFocalLength=70,MinFStop=2.8,MaxFStop=22,DiaphragmBladeCount=7))
+-LensPresets=(Name="70-200mm Zoom f/2.8",LensSettings=(MinFocalLength=70,MaxFocalLength=200,MinFStop=2.8,MaxFStop=22,DiaphragmBladeCount=7))
+-LensPresets=(Name="Universal Zoom",LensSettings=(MinFocalLength=4,MaxFocalLength=1000,MinFStop=1.2,MaxFStop=22,DiaphragmBladeCount=7))
++LensPresets=(Name="Zoom",LensSettings=(MinFocalLength=8.000000,MaxFocalLength=1000.000000,MinFStop=1.200000,MaxFStop=22.000000,MinimumFocusDistance=15.000000,SqueezeFactor=1.000000,DiaphragmBladeCount=7))
++LensPresets=(Name="14mm",LensSettings=(MinFocalLength=14.000000,MaxFocalLength=14.000000,MinFStop=1.400000,MaxFStop=22.000000,MinimumFocusDistance=15.000000,SqueezeFactor=1.000000,DiaphragmBladeCount=7))
++LensPresets=(Name="16mm",LensSettings=(MinFocalLength=16.000000,MaxFocalLength=16.000000,MinFStop=1.400000,MaxFStop=22.000000,MinimumFocusDistance=15.000000,SqueezeFactor=1.000000,DiaphragmBladeCount=7))
++LensPresets=(Name="18mm",LensSettings=(MinFocalLength=18.000000,MaxFocalLength=18.000000,MinFStop=1.400000,MaxFStop=22.000000,MinimumFocusDistance=15.000000,SqueezeFactor=1.000000,DiaphragmBladeCount=7))
++LensPresets=(Name="21mm",LensSettings=(MinFocalLength=21.000000,MaxFocalLength=21.000000,MinFStop=1.400000,MaxFStop=22.000000,MinimumFocusDistance=15.000000,SqueezeFactor=1.000000,DiaphragmBladeCount=7))
++LensPresets=(Name="25mm",LensSettings=(MinFocalLength=25.000000,MaxFocalLength=25.000000,MinFStop=1.400000,MaxFStop=22.000000,MinimumFocusDistance=15.000000,SqueezeFactor=1.000000,DiaphragmBladeCount=7))
++LensPresets=(Name="27mm",LensSettings=(MinFocalLength=27.000000,MaxFocalLength=27.000000,MinFStop=1.400000,MaxFStop=22.000000,MinimumFocusDistance=15.000000,SqueezeFactor=1.000000,DiaphragmBladeCount=7))
++LensPresets=(Name="32mm",LensSettings=(MinFocalLength=32.000000,MaxFocalLength=32.000000,MinFStop=1.400000,MaxFStop=22.000000,MinimumFocusDistance=15.000000,SqueezeFactor=1.000000,DiaphragmBladeCount=7))
++LensPresets=(Name="35mm",LensSettings=(MinFocalLength=35.000000,MaxFocalLength=35.000000,MinFStop=1.400000,MaxFStop=22.000000,MinimumFocusDistance=15.000000,SqueezeFactor=1.000000,DiaphragmBladeCount=7))
++LensPresets=(Name="40mm",LensSettings=(MinFocalLength=40.000000,MaxFocalLength=40.000000,MinFStop=1.400000,MaxFStop=22.000000,MinimumFocusDistance=15.000000,SqueezeFactor=1.000000,DiaphragmBladeCount=7))
++LensPresets=(Name="50mm",LensSettings=(MinFocalLength=50.000000,MaxFocalLength=50.000000,MinFStop=1.400000,MaxFStop=22.000000,MinimumFocusDistance=15.000000,SqueezeFactor=1.000000,DiaphragmBladeCount=7))
++LensPresets=(Name="65mm",LensSettings=(MinFocalLength=65.000000,MaxFocalLength=65.000000,MinFStop=1.400000,MaxFStop=22.000000,MinimumFocusDistance=15.000000,SqueezeFactor=1.000000,DiaphragmBladeCount=7))
++LensPresets=(Name="75mm",LensSettings=(MinFocalLength=75.000000,MaxFocalLength=75.000000,MinFStop=1.400000,MaxFStop=22.000000,MinimumFocusDistance=15.000000,SqueezeFactor=1.000000,DiaphragmBladeCount=7))
++LensPresets=(Name="100mm",LensSettings=(MinFocalLength=100.000000,MaxFocalLength=100.000000,MinFStop=1.400000,MaxFStop=22.000000,MinimumFocusDistance=15.000000,SqueezeFactor=1.000000,DiaphragmBladeCount=7))
++LensPresets=(Name="135mm",LensSettings=(MinFocalLength=135.000000,MaxFocalLength=135.000000,MinFStop=1.400000,MaxFStop=22.000000,MinimumFocusDistance=15.000000,SqueezeFactor=1.000000,DiaphragmBladeCount=7))
++LensPresets=(Name="150mm",LensSettings=(MinFocalLength=150.000000,MaxFocalLength=150.000000,MinFStop=1.400000,MaxFStop=22.000000,MinimumFocusDistance=15.000000,SqueezeFactor=1.000000,DiaphragmBladeCount=7))
++LensPresets=(Name="180mm",LensSettings=(MinFocalLength=180.000000,MaxFocalLength=180.000000,MinFStop=1.400000,MaxFStop=22.000000,MinimumFocusDistance=15.000000,SqueezeFactor=1.000000,DiaphragmBladeCount=7))
+DefaultFilmbackPreset=16:9 Digital Film
+-FilmbackPresets=(Name="16:9 Film",FilmbackSettings=(SensorWidth=24.00,SensorHeight=13.5))
+-FilmbackPresets=(Name="16:9 Digital Film",FilmbackSettings=(SensorWidth=23.76,SensorHeight=13.365))
+-FilmbackPresets=(Name="16:9 DSLR",FilmbackSettings=(SensorWidth=36,SensorHeight=20.25))
+-FilmbackPresets=(Name="Super 8mm",FilmbackSettings=(SensorWidth=5.79,SensorHeight=4.01))
+-FilmbackPresets=(Name="Super 16mm",FilmbackSettings=(SensorWidth=12.52,SensorHeight=7.58))
+-FilmbackPresets=(Name="Super 35mm",FilmbackSettings=(SensorWidth=24.89,SensorHeight=18.66))
+-FilmbackPresets=(Name="35mm Academy",FilmbackSettings=(SensorWidth=21.946,SensorHeight=16.002))
+-FilmbackPresets=(Name="35mm Full Aperture",FilmbackSettings=(SensorWidth=24.892,SensorHeight=18.9121))
+-FilmbackPresets=(Name="35mm VistaVision",FilmbackSettings=(SensorWidth=37.719,SensorHeight=25.146))
+-FilmbackPresets=(Name="IMAX 70mm",FilmbackSettings=(SensorWidth=70.41,SensorHeight=56.63))
+-FilmbackPresets=(Name="APS-C (Canon)",FilmbackSettings=(SensorWidth=22.2,SensorHeight=14.8))
+-FilmbackPresets=(Name="Full Frame DSLR",FilmbackSettings=(SensorWidth=36,SensorHeight=24))
+-FilmbackPresets=(Name="Micro Four Thirds",FilmbackSettings=(SensorWidth=17.3,SensorHeight=13))
++FilmbackPresets=(Name="16:9 Film",FilmbackSettings=(SensorWidth=24.000000,SensorHeight=13.500000,SensorAspectRatio=1.330000))
++FilmbackPresets=(Name="16:9 Digital Film",FilmbackSettings=(SensorWidth=23.760000,SensorHeight=13.365000,SensorAspectRatio=1.330000))
++FilmbackPresets=(Name="16:9 DSLR",FilmbackSettings=(SensorWidth=36.000000,SensorHeight=20.250000,SensorAspectRatio=1.330000))
++FilmbackPresets=(Name="Super 8mm",FilmbackSettings=(SensorWidth=5.790000,SensorHeight=4.010000,SensorAspectRatio=1.330000))
++FilmbackPresets=(Name="Super 16mm",FilmbackSettings=(SensorWidth=12.520000,SensorHeight=7.580000,SensorAspectRatio=1.330000))
++FilmbackPresets=(Name="Super 35mm",FilmbackSettings=(SensorWidth=24.889999,SensorHeight=18.660000,SensorAspectRatio=1.330000))
++FilmbackPresets=(Name="35mm Academy",FilmbackSettings=(SensorWidth=21.945999,SensorHeight=16.002001,SensorAspectRatio=1.330000))
++FilmbackPresets=(Name="35mm Full Aperture",FilmbackSettings=(SensorWidth=24.892000,SensorHeight=18.912100,SensorAspectRatio=1.330000))
++FilmbackPresets=(Name="35mm VistaVision",FilmbackSettings=(SensorWidth=37.719002,SensorHeight=25.146000,SensorAspectRatio=1.330000))
++FilmbackPresets=(Name="IMAX 70mm",FilmbackSettings=(SensorWidth=70.410004,SensorHeight=56.630001,SensorAspectRatio=1.330000))
++FilmbackPresets=(Name="APS-C (Canon)",FilmbackSettings=(SensorWidth=22.200001,SensorHeight=14.800000,SensorAspectRatio=1.330000))
++FilmbackPresets=(Name="Full Frame DSLR",FilmbackSettings=(SensorWidth=36.000000,SensorHeight=24.000000,SensorAspectRatio=1.330000))
++FilmbackPresets=(Name="Micro Four Thirds",FilmbackSettings=(SensorWidth=17.299999,SensorHeight=13.000000,SensorAspectRatio=1.330000))
+DefaultCropPresetName=No Crop
+-CropPresets=(Name="No Crop",CropSettings=(AspectRatio=0))
+-CropPresets=(Name="1.33 (4:3)",CropSettings=(AspectRatio=1.333333))
+-CropPresets=(Name="1.77 (16:9)",CropSettings=(AspectRatio=1.777778))
+-CropPresets=(Name="2.39",CropSettings=(AspectRatio=2.39))
++CropPresets=(Name="No Crop",CropSettings=(AspectRatio=0.000000))
++CropPresets=(Name="1.33 (4:3)",CropSettings=(AspectRatio=1.333333))
++CropPresets=(Name="1.77 (16:9)",CropSettings=(AspectRatio=1.777778))
++CropPresets=(Name="2.39",CropSettings=(AspectRatio=2.390000))
