[Backout] - CL22732712

Original CL Desc
-----------------------------------------------------------------
Corrected Virtual Camera config files

diff --git a/Engine/Plugins/Experimental/VirtualCamera/Config/BaseEngine.ini b/Engine/Plugins/Experimental/VirtualCamera/Config/BaseEngine.ini
new file mode 100644
--- /dev/null
+++ b/Engine/Plugins/Experimental/VirtualCamera/Config/BaseEngine.ini
@@ -0,0 +1,12 @@
+[CoreRedirects]
++ClassRedirects=(OldName="/Script/VCamCore.VCamCoreSubsystem",NewName="/Script/VCamInput.VCamInputSubsystem")
++ClassRedirects=(OldName="/Script/VirtualCamera.RadialSlider",NewName="/Script/AdvancedWidgets.RadialSlider")
++FunctionRedirects=(OldName="VCamComponent.ClearModifiers",NewName="VCamComponent.RemoveAllModifiers")
++FunctionRedirects=(OldName="VCamComponent.FindModifierByName",NewName="VCamComponent.GetModifierByName")
++FunctionRedirects=(OldName="VCamComponent.FindModifiersByClass",NewName="VCamComponent.GetModifiersByClass")
++FunctionRedirects=(OldName="VCamComponent.FindModifierByInterface",NewName="VCamComponent.GetModifiersByInterface")
++PackageRedirects=(OldName="/VirtualCamera/Editor", NewName="/VirtualCamera/Deprecated/Editor", MatchSubstring=true)
++PackageRedirects=(OldName="/VirtualCamera/OperatorPanel", NewName="/VirtualCamera/Deprecated/OperatorPanel", MatchSubstring=true)
++PackageRedirects=(OldName="/VirtualCamera/TestMaps", NewName="/VirtualCamera/Deprecated/TestMaps", MatchSubstring=true)
++PackageRedirects=(OldName="/VirtualCamera/V2", NewName="/VirtualCamera/Deprecated/V2", MatchSubstring=true)
++PackageRedirects=(OldName="/VirtualCamera/VCamCore", NewName="/VirtualCamera/Deprecated/VCamCore", MatchSubstring=true)
diff --git a/Engine/Plugins/Experimental/VirtualCamera/Config/Engine.ini b/Engine/Plugins/Experimental/VirtualCamera/Config/Engine.ini
deleted file mode 100644
