Fix non-unity compile errors for VCamOutputProviderBase

diff --git a/Engine/Plugins/Experimental/VirtualCamera/Source/VCamCore/Private/VCamOutputProviderBase.cpp b/Engine/Plugins/Experimental/VirtualCamera/Source/VCamCore/Private/VCamOutputProviderBase.cpp
--- a/Engine/Plugins/Experimental/VirtualCamera/Source/VCamCore/Private/VCamOutputProviderBase.cpp
+++ b/Engine/Plugins/Experimental/VirtualCamera/Source/VCamCore/Private/VCamOutputProviderBase.cpp
@@ -2,6 +2,7 @@
 
 #include "VCamOutputProviderBase.h"
 
+#include "Engine/Engine.h"
 #include "VCamModifierInterface.h"
 #include "VCamComponent.h"
 #include "UObject/UObjectBaseUtility.h"
