Reorganize files and add namespaces (UE code style standard)

diff --git a/Engine/Plugins/Experimental/VirtualCamera/Source/VCamCoreEditor/Private/Customizations/VCamInputProfileCustomization.cpp b/Engine/Plugins/Experimental/VirtualCamera/Source/VCamCoreEditor/Private/Customizations/VCamInputProfileCustomization.cpp
new file mode 100644
--- /dev/null
+++ b/Engine/Plugins/Experimental/VirtualCamera/Source/VCamCoreEditor/Private/Customizations/VCamInputProfileCustomization.cpp
@@ -0,0 +1,208 @@
+// Copyright Epic Games, Inc. All Rights Reserved.
+
+#include "VCamInputProfileCustomization.h"
+
+#include "DetailLayoutBuilder.h"
+#include "DetailWidgetRow.h"
+#include "IDetailChildrenBuilder.h"
+#include "VCamInputSettings.h"
+#include "ScopedTransaction.h"
+
+#define LOCTEXT_NAMESPACE "VCamInputProfileCustomization"
+
+namespace UE::VCamCoreEditor::Private
+{
+	FVCamInputProfileCustomization::FVCamInputProfileCustomization()
+	{
+		BuildProfileComboList();
+	}
+
+	TSharedRef<IPropertyTypeCustomization> FVCamInputProfileCustomization::MakeInstance()
+	{
+		return MakeShared<FVCamInputProfileCustomization>();
+	}
+
+	void FVCamInputProfileCustomization::CustomizeHeader(TSharedRef<IPropertyHandle> PropertyHandle, FDetailWidgetRow& HeaderRow, IPropertyTypeCustomizationUtils& CustomizationUtils)
+	{
+		HeaderRow.
+			NameContent()
+			[
+				PropertyHandle->CreatePropertyNameWidget()
+			];
+
+		// We only want the dropdown list outside of the settings class as the settings class is the thing
+		// defining the presets we use for the dropdown
+		const bool bInSettingsClass = PropertyHandle->GetOuterBaseClass() == UVCamInputSettings::StaticClass();
+
+		if (!bInSettingsClass)
+		{
+			HeaderRow.
+				ValueContent()
+				.MaxDesiredWidth(0.f)
+				[
+					SAssignNew(ProfileComboBox, SComboBox< TSharedPtr<FString> >)
+					.OptionsSource(&ProfileComboList)
+					.OnGenerateWidget(this, &FVCamInputProfileCustomization::MakeProfileComboWidget)
+					.OnSelectionChanged(this, &FVCamInputProfileCustomization::OnProfileChanged)
+					.OnComboBoxOpening(this, &FVCamInputProfileCustomization::BuildProfileComboList)
+					.IsEnabled(this, &FVCamInputProfileCustomization::IsProfileEnabled)
+					.ContentPadding(2)
+					.Content()
+					[
+						SNew(STextBlock)
+						.Text(this, &FVCamInputProfileCustomization::GetProfileComboBoxContent)
+						.Font(IDetailLayoutBuilder::GetDetailFont())
+						.ToolTipText(this, &FVCamInputProfileCustomization::GetProfileComboBoxContent)
+					]
+				];
+		}
+	}
+
+	void FVCamInputProfileCustomization::CustomizeChildren(TSharedRef<IPropertyHandle> StructPropertyHandle, IDetailChildrenBuilder& ChildBuilder, IPropertyTypeCustomizationUtils& StructCustomizationUtils)
+	{
+		CachedStructPropertyHandle = StructPropertyHandle;
+		
+		// Retrieve structure's child properties
+		uint32 NumChildren;
+		StructPropertyHandle->GetNumChildren( NumChildren );	
+		const FName MappableKeyOverridesPropertyName = GET_MEMBER_NAME_CHECKED(FVCamInputProfile, MappableKeyOverrides);
+		for(uint32 ChildIndex = 0; ChildIndex < NumChildren; ++ChildIndex)
+		{
+			TSharedPtr<IPropertyHandle> ChildProperty = StructPropertyHandle->GetChildHandle(ChildIndex);
+			ChildBuilder.AddProperty(ChildProperty.ToSharedRef());
+
+			if (ChildProperty->GetProperty() && ChildProperty->GetProperty()->GetFName() == MappableKeyOverridesPropertyName)
+			{
+				MappableKeyOverridesHandle = ChildProperty;
+			}
+		}
+	}
+
+	void FVCamInputProfileCustomization::BuildProfileComboList()
+	{
+		const UVCamInputSettings* VCamInputSettings = GetDefault<UVCamInputSettings>();
+
+		if (!VCamInputSettings)
+		{
+			return;
+		}
+		
+		const TArray<FName>& Presets = VCamInputSettings->GetInputProfileNames();
+		
+		ProfileComboList.Empty(Presets.Num() + 1);
+		ProfileComboList.Add(MakeShared<FString>(TEXT("Custom...")));
+
+		// put all presets in the list
+		for (const FName& Preset : Presets)
+		{
+			ProfileComboList.Add(MakeShared<FString>(Preset.ToString()));
+		}
+	}
+
+	bool FVCamInputProfileCustomization::IsProfileEnabled() const
+	{
+		if (MappableKeyOverridesHandle.IsValid())
+		{
+			return MappableKeyOverridesHandle->IsEditable() && FSlateApplication::Get().GetNormalExecutionAttribute().Get();
+		}
+		return false;
+	}
+
+	void FVCamInputProfileCustomization::OnProfileChanged(TSharedPtr<FString> NewSelection, ESelectInfo::Type SelectInfo)
+	{
+		const UVCamInputSettings* VCamInputSettings = GetDefault<UVCamInputSettings>();
+
+		if (!VCamInputSettings)
+		{
+			return;
+		}
+		
+		// Ignore any changes that come from code
+		if (SelectInfo != ESelectInfo::Direct && CachedStructPropertyHandle.IsValid() && MappableKeyOverridesHandle.IsValid())
+		{
+			const FString NewPresetName = *NewSelection.Get();
+
+			if (const FVCamInputProfile* InputProfile = VCamInputSettings->InputProfiles.Find(FName(NewPresetName)))
+			{
+				const FScopedTransaction Transaction(LOCTEXT("ChangeVCamInputProfile", "Change VCam Input Profile"));
+				
+				FStructProperty* StructProperty = CastFieldChecked<FStructProperty>(CachedStructPropertyHandle->GetProperty());
+				FMapProperty* MapProperty = CastFieldChecked<FMapProperty>(MappableKeyOverridesHandle->GetProperty());
+
+				TArray<UObject*> OuterObjects;
+				MappableKeyOverridesHandle->GetOuterObjects(OuterObjects);
+
+				CachedStructPropertyHandle->NotifyPreChange();
+				
+				for (UObject* OuterObject : OuterObjects)
+				{				
+					void* StructValuePtr = StructProperty->ContainerPtrToValuePtr<void*>(OuterObject);
+					void* MapValuePtr = MapProperty->ContainerPtrToValuePtr<void*>(StructValuePtr);
+					
+					FScriptMapHelper MapHelper(MapProperty, MapValuePtr);
+
+					MapHelper.EmptyValues();
+
+					for (const TPair<FName, FKey>& MappingPair : InputProfile->MappableKeyOverrides)
+					{
+						MapHelper.AddPair(&MappingPair.Key, &MappingPair.Value);
+					}
+				}
+
+				CachedStructPropertyHandle->NotifyPostChange(EPropertyChangeType::ValueSet);
+				CachedStructPropertyHandle->NotifyFinishedChangingProperties();
+			}
+		}
+	}
+
+	TSharedRef<SWidget> FVCamInputProfileCustomization::MakeProfileComboWidget(TSharedPtr<FString> InItem)
+	{
+		return SNew(STextBlock)
+			.Text(FText::FromString(*InItem))
+			.Font(IDetailLayoutBuilder::GetDetailFont());
+	}
+
+	FText FVCamInputProfileCustomization::GetProfileComboBoxContent() const
+	{
+		// just test one variable for multiple selection
+		void* StructValueData;
+		if (CachedStructPropertyHandle->GetValueData(StructValueData) == FPropertyAccess::Result::MultipleValues)
+		{
+			// multiple selection
+			return LOCTEXT("MultipleValues", "Multiple Values");
+		}
+		const TSharedPtr<FString> ProfileString = GetProfileString(StructValueData);
+		return ProfileString.IsValid() ? FText::AsCultureInvariant(*ProfileString) : LOCTEXT("NoProfileFound", "Error");
+	}
+
+	TSharedPtr<FString> FVCamInputProfileCustomization::GetProfileString(void* StructValueData) const
+	{
+		FVCamInputProfile* InputProfilePtr = static_cast<FVCamInputProfile*>(StructValueData);
+		const UVCamInputSettings* VCamInputSettings = GetDefault<UVCamInputSettings>();
+
+		if (InputProfilePtr && VCamInputSettings)
+		{
+			for (const TPair<FName, FVCamInputProfile>& InputProfilePair : VCamInputSettings->InputProfiles)
+			{
+				const FName& InputProfileName = InputProfilePair.Key;
+				const FVCamInputProfile& InputProfileValue = InputProfilePair.Value;
+				if (InputProfileValue == *InputProfilePtr)
+				{
+					const TSharedPtr<FString>* FoundString = ProfileComboList.FindByPredicate([InputProfileString = InputProfileName.ToString()](const TSharedPtr<FString> ListEntry)
+					{
+						return ListEntry.IsValid() ? *ListEntry == InputProfileString : false;
+					});
+					
+					if (FoundString)
+					{
+						return *FoundString;
+					}
+					break;
+				}
+			}
+		}
+		return ProfileComboList[0];
+	}
+}
+
+#undef LOCTEXT_NAMESPACE
\ No newline at end of file
diff --git a/Engine/Plugins/Experimental/VirtualCamera/Source/VCamCoreEditor/Private/Customizations/VCamInputProfileCustomization.h b/Engine/Plugins/Experimental/VirtualCamera/Source/VCamCoreEditor/Private/Customizations/VCamInputProfileCustomization.h
new file mode 100644
--- /dev/null
+++ b/Engine/Plugins/Experimental/VirtualCamera/Source/VCamCoreEditor/Private/Customizations/VCamInputProfileCustomization.h
@@ -0,0 +1,46 @@
+// Copyright Epic Games, Inc. All Rights Reserved.
+
+#pragma once
+
+#include "Containers/Array.h"
+#include "Containers/UnrealString.h"
+#include "IPropertyTypeCustomization.h"
+#include "Internationalization/Text.h"
+#include "Templates/SharedPointer.h"
+#include "Types/SlateEnums.h"
+#include "Widgets/Input/SComboBox.h"
+
+class FDetailWidgetRow;
+class IPropertyHandle;
+class SWidget;
+
+namespace UE::VCamCoreEditor::Private
+{
+	class FVCamInputProfileCustomization : public IPropertyTypeCustomization
+	{
+	public:
+		FVCamInputProfileCustomization();
+
+		static TSharedRef<IPropertyTypeCustomization> MakeInstance();
+
+		/** IPropertyTypeCustomization instance */
+		virtual void CustomizeHeader(TSharedRef<IPropertyHandle> PropertyHandle, FDetailWidgetRow& HeaderRow, IPropertyTypeCustomizationUtils& CustomizationUtils) override;
+		virtual void CustomizeChildren( TSharedRef<IPropertyHandle> StructPropertyHandle, class IDetailChildrenBuilder& ChildBuilder, IPropertyTypeCustomizationUtils& StructCustomizationUtils ) override;
+
+	protected:
+		TSharedPtr<IPropertyHandle> CachedStructPropertyHandle;
+		TSharedPtr<IPropertyHandle> MappableKeyOverridesHandle;
+	
+		TSharedPtr<SComboBox<TSharedPtr<FString>>> ProfileComboBox;
+		TArray<TSharedPtr<FString>> ProfileComboList;
+
+		void BuildProfileComboList();
+
+		bool IsProfileEnabled() const;
+		void OnProfileChanged(TSharedPtr<FString> NewSelection, ESelectInfo::Type SelectInfo);
+		TSharedRef<SWidget> MakeProfileComboWidget(TSharedPtr<FString> InItem);
+		
+		FText GetProfileComboBoxContent() const;
+		TSharedPtr<FString> GetProfileString(void* StructValueData) const;
+	};
+}
\ No newline at end of file
diff --git a/Engine/Plugins/Experimental/VirtualCamera/Source/VCamCoreEditor/Private/Factories/VCamModifierFactory.cpp b/Engine/Plugins/Experimental/VirtualCamera/Source/VCamCoreEditor/Private/Factories/VCamModifierFactory.cpp
new file mode 100644
--- /dev/null
+++ b/Engine/Plugins/Experimental/VirtualCamera/Source/VCamCoreEditor/Private/Factories/VCamModifierFactory.cpp
@@ -0,0 +1,56 @@
+// Copyright Epic Games, Inc. All Rights Reserved.
+
+#include "Factories/VCamModifierFactory.h"
+#include "VCamModifier.h"
+
+#include "AssetToolsModule.h"
+#include "Kismet2/BlueprintEditorUtils.h"
+#include "Kismet2/KismetEditorUtilities.h"
+
+#define LOCTEXT_NAMESPACE "VCamModifierFactory"
+
+UVCamModifierFactory::UVCamModifierFactory()
+{
+	SupportedClass = UBlueprint::StaticClass();
+	ParentClass = UVCamBlueprintModifier::StaticClass();
+
+	bCreateNew = true;
+	bEditAfterNew = true;
+}
+
+FText UVCamModifierFactory::GetDisplayName() const
+{
+	return LOCTEXT("DisplayName", "VCam Modifier");
+}
+
+FText UVCamModifierFactory::GetToolTip() const
+{
+	return LOCTEXT("Tooltip", "Modifiers can manipulate the camera associated with a VCam Component. You can create custom modifiers with Blueprints or C++ and add them to the modifier stack of a VCam Component to layer different effects.");
+}
+
+UObject* UVCamModifierFactory::FactoryCreateNew(UClass* Class, UObject* InParent, FName Name, EObjectFlags Flags, UObject* Context, FFeedbackContext* Warn)
+{
+	UBlueprint* ModifierBlueprint = nullptr;
+	if (ensure(SupportedClass == Class))
+	{
+		ensure(0 != (RF_Public & Flags));
+		ModifierBlueprint = FKismetEditorUtilities::CreateBlueprint(ParentClass, InParent, Name, BPTYPE_Normal, UBlueprint::StaticClass(), UBlueprintGeneratedClass::StaticClass(), NAME_None);
+		if (TSubclassOf<UObject> GeneratedClass = ModifierBlueprint->GeneratedClass)
+		{
+			if (UVCamBlueprintModifier* DefaultSubject = GeneratedClass->GetDefaultObject<UVCamBlueprintModifier>())
+			{
+				DefaultSubject->InputMappingContext = InputMappingContext;
+			}
+		}
+		FBlueprintEditorUtils::MarkBlueprintAsModified(ModifierBlueprint);
+	}
+	return ModifierBlueprint;	
+}
+
+uint32 UVCamModifierFactory::GetMenuCategories() const
+{
+	IAssetTools& AssetTools = FModuleManager::LoadModuleChecked<FAssetToolsModule>("AssetTools").Get();
+	return AssetTools.RegisterAdvancedAssetCategory("VirtualCamera", LOCTEXT("AssetCategoryName", "VCam"));
+}
+
+#undef LOCTEXT_NAMESPACE
diff --git a/Engine/Plugins/Experimental/VirtualCamera/Source/VCamCoreEditor/Private/Factories/VCamObjectWithInputFactory.cpp b/Engine/Plugins/Experimental/VirtualCamera/Source/VCamCoreEditor/Private/Factories/VCamObjectWithInputFactory.cpp
new file mode 100644
--- /dev/null
+++ b/Engine/Plugins/Experimental/VirtualCamera/Source/VCamCoreEditor/Private/Factories/VCamObjectWithInputFactory.cpp
@@ -0,0 +1,149 @@
+// Copyright Epic Games, Inc. All Rights Reserved.
+
+#include "Factories/VCamObjectWithInputFactory.h"
+
+#include "AssetThumbnail.h"
+#include "Editor.h"
+#include "PropertyCustomizationHelpers.h"
+
+#include "Widgets/SBoxPanel.h"
+#include "Widgets/SWindow.h"
+#include "Widgets/Input/SButton.h"
+#include "Widgets/Layout/SBorder.h"
+#include "Widgets/Text/STextBlock.h"
+
+#define LOCTEXT_NAMESPACE "VCamObjectFactories"
+
+bool UVCamObjectWithInputFactory::ConfigureProperties()
+{
+	class FVCamWidgetFactoryUI : public TSharedFromThis<FVCamWidgetFactoryUI>
+	{
+	public:
+		FReply OnCreate()
+		{
+			if (PickerWindow.IsValid())
+			{
+				PickerWindow->RequestDestroyWindow();
+			}
+			bWasCanceled= false;
+			return FReply::Handled();
+		}
+
+		bool WasCanceled() const
+		{
+			return bWasCanceled;
+		}
+
+		void OpenMappingSelector(UVCamObjectWithInputFactory* Factory)
+		{
+			AssetThumbnailPool = MakeShared<FAssetThumbnailPool>(1);
+			PickerWindow = SNew(SWindow)
+				.Title(LOCTEXT("PickerTitle", "Select Input Mapping Context"))
+				.ClientSize(FVector2D(450, 150))
+				.SupportsMinimize(false)
+				.SupportsMaximize(false)
+				[
+					SNew(SBorder)
+					.BorderImage(FAppStyle::GetBrush("Menu.Background"))
+					.Padding(10)
+					[
+						SNew(SVerticalBox)
+						+ SVerticalBox::Slot()
+						.AutoHeight()
+						[
+							SNew(STextBlock)
+							.Text(LOCTEXT("DialogText", "Your class needs an Input Mapping Context if you want to accept hardware input. You can add or change this selection later in the Class Defaults."))
+							.AutoWrapText(true)
+						]
+						+ SVerticalBox::Slot()
+						.FillHeight(1.f)
+						[
+							SNew(SHorizontalBox)
+							+ SHorizontalBox::Slot()
+							.AutoWidth()
+							.HAlign(HAlign_Left)
+							.VAlign(VAlign_Center)
+							.Padding(FMargin(0, 0, 20, 0))
+							[
+								SNew(STextBlock)
+								.Text(LOCTEXT("InputMappingContext", "Input Mapping Context"))
+							]
+							+ SHorizontalBox::Slot()
+							.FillWidth(1.f)
+							.VAlign(VAlign_Center)
+							[
+								SNew(SObjectPropertyEntryBox)
+								.DisplayThumbnail(true)
+								.ThumbnailPool(AssetThumbnailPool)
+								.AllowClear(true)
+								.DisplayUseSelected(false)
+								.DisplayBrowse(false)
+								.AllowedClass(UInputMappingContext::StaticClass())
+								.ObjectPath_Lambda([Factory]() -> FString
+								{
+									return (Factory && Factory->InputMappingContext) ? Factory->InputMappingContext->GetPathName() : TEXT("None");
+								})
+								.OnObjectChanged_Lambda([Factory](const FAssetData& AssetData) -> void
+								{
+									if (Factory)
+									{
+										Factory->InputMappingContext = Cast<UInputMappingContext>(AssetData.GetAsset());
+									}
+								})
+							]
+						]
+						+ SVerticalBox::Slot()
+						.HAlign(HAlign_Right)
+						.VAlign(VAlign_Bottom)
+						.AutoHeight()
+						[
+							SNew(SHorizontalBox)
+							+ SHorizontalBox::Slot()
+							.AutoWidth()
+							.Padding(FMargin(5,0))
+							[
+								SNew(SButton)
+								.Text(LOCTEXT("OK", "OK"))
+								.ToolTipText(LOCTEXT("OkTooltip", "Continues creation using the specified Input Mapping Context"))
+								.OnClicked(this, &FVCamWidgetFactoryUI::OnCreate)
+							]
+							+ SHorizontalBox::Slot()
+							.AutoWidth()
+							.Padding(FMargin(5,0))
+							[
+								SNew(SButton)
+								.Text(LOCTEXT("Skip", "Skip"))
+								.ToolTipText(LOCTEXT("SkipTooltip", "Continues creation without assigning an Input Mapping Context"))
+								.OnClicked_Lambda([Factory, this]() -> FReply
+								{
+									// If we skipped then create the widget but with no IMC set
+									if (Factory)
+									{
+										Factory->InputMappingContext = nullptr;
+									}
+									return OnCreate();
+								})
+							]
+						]
+					]
+				];
+			if (GEditor)
+			{
+				GEditor->EditorAddModalWindow(PickerWindow.ToSharedRef());
+			}
+			PickerWindow.Reset();
+		}
+
+	private:
+		TSharedPtr<SWindow> PickerWindow;
+		TSharedPtr<FAssetThumbnailPool> AssetThumbnailPool;
+		bool bWasCanceled = true;
+	};
+
+	const TSharedRef<FVCamWidgetFactoryUI> InputMappingSelector = MakeShared<FVCamWidgetFactoryUI>();
+	InputMappingSelector->OpenMappingSelector(this);
+	
+	return !InputMappingSelector->WasCanceled();
+}
+
+#undef LOCTEXT_NAMESPACE
\ No newline at end of file
diff --git a/Engine/Plugins/Experimental/VirtualCamera/Source/VCamCoreEditor/Private/Factories/VCamWidgetFactory.cpp b/Engine/Plugins/Experimental/VirtualCamera/Source/VCamCoreEditor/Private/Factories/VCamWidgetFactory.cpp
new file mode 100644
--- /dev/null
+++ b/Engine/Plugins/Experimental/VirtualCamera/Source/VCamCoreEditor/Private/Factories/VCamWidgetFactory.cpp
@@ -0,0 +1,57 @@
+// Copyright Epic Games, Inc. All Rights Reserved.
+
+#include "Factories/VCamWidgetFactory.h"
+#include "UI/VCamWidget.h"
+
+#include "AssetToolsModule.h"
+#include "Kismet2/BlueprintEditorUtils.h"
+#include "Kismet2/KismetEditorUtilities.h"
+#include "WidgetBlueprint.h"
+
+#define LOCTEXT_NAMESPACE "VCamWidgetFactory"
+
+UVCamWidgetFactory::UVCamWidgetFactory()
+{
+	SupportedClass = UBlueprint::StaticClass();
+	ParentClass = UVCamWidget::StaticClass();
+
+	bCreateNew = true;
+	bEditAfterNew = true;
+}
+
+FText UVCamWidgetFactory::GetDisplayName() const
+{
+	return LOCTEXT("DisplayName", "VCam Widget");
+}
+
+FText UVCamWidgetFactory::GetToolTip() const
+{
+	return LOCTEXT("Tooltip", "A wrapper widget class that contains a set of VCam Connections\n\nIf you add a widget deriving from VCam Widget to an Overlay Widget for a VCam Output Provider then when the Overlay is created by the Provider it will also call InitializeConnections with the owning VCam Component.");
+}
+
+UObject* UVCamWidgetFactory::FactoryCreateNew(UClass* Class, UObject* InParent, FName Name, EObjectFlags Flags, UObject* Context, FFeedbackContext* Warn)
+{
+	UBlueprint* WidgetBlueprint = nullptr;
+	if (ensure(SupportedClass == Class))
+	{
+		ensure(0 != (RF_Public & Flags));
+		WidgetBlueprint = FKismetEditorUtilities::CreateBlueprint(ParentClass, InParent, Name, BPTYPE_Normal, UWidgetBlueprint::StaticClass(), UWidgetBlueprintGeneratedClass::StaticClass(), NAME_None);
+		if (TSubclassOf<UObject> GeneratedClass = WidgetBlueprint->GeneratedClass)
+		{
+			if (UVCamWidget* DefaultSubject = GeneratedClass->GetDefaultObject<UVCamWidget>())
+			{
+				DefaultSubject->InputMappingContext = InputMappingContext;
+			}
+		}
+		FBlueprintEditorUtils::MarkBlueprintAsModified(WidgetBlueprint);
+	}
+	return WidgetBlueprint;	
+}
+
+uint32 UVCamWidgetFactory::GetMenuCategories() const
+{
+	IAssetTools& AssetTools = FModuleManager::LoadModuleChecked<FAssetToolsModule>("AssetTools").Get();
+	return AssetTools.RegisterAdvancedAssetCategory("VirtualCamera", LOCTEXT("AssetCategoryName", "VCam"));
+}
+
+#undef LOCTEXT_NAMESPACE
diff --git a/Engine/Plugins/Experimental/VirtualCamera/Source/VCamCoreEditor/Private/VCamCoreEditorModule.cpp b/Engine/Plugins/Experimental/VirtualCamera/Source/VCamCoreEditor/Private/VCamCoreEditorModule.cpp
--- a/Engine/Plugins/Experimental/VirtualCamera/Source/VCamCoreEditor/Private/VCamCoreEditorModule.cpp
+++ b/Engine/Plugins/Experimental/VirtualCamera/Source/VCamCoreEditor/Private/VCamCoreEditorModule.cpp
@@ -2,39 +2,44 @@
 
 #include "Modules/ModuleInterface.h"
 #include "Modules/ModuleManager.h"
+
+#include "Customizations/VCamInputProfileCustomization.h"
 #include "Interfaces/IPluginManager.h"
 #include "Kismet2/KismetEditorUtilities.h"
 #include "PropertyEditorModule.h"
+#include "VCamComponent.h"
 #include "VCamModifier.h"
-#include "VCamInputProfileCustomization.h"
 
-class FVCamCoreEditorModule : public IModuleInterface
+namespace UE::VCamCoreEditor::Private
 {
-public:
+	class FVCamCoreEditorModule : public IModuleInterface
+	{
+	public:
 
-	// IModuleInterface interface
+		// IModuleInterface interface
 
-	virtual void StartupModule() override
-	{
-		RegisterAutoGeneratedDefaultEvents();
+		virtual void StartupModule() override
+		{
+			RegisterAutoGeneratedDefaultEvents();
 		
-		FPropertyEditorModule& PropertyModule = FModuleManager::GetModuleChecked<FPropertyEditorModule>("PropertyEditor");
-		PropertyModule.RegisterCustomPropertyTypeLayout("VCamInputProfile", FOnGetPropertyTypeCustomizationInstance::CreateStatic(&FVCamInputProfileCustomization::MakeInstance));
-	};
+			FPropertyEditorModule& PropertyModule = FModuleManager::GetModuleChecked<FPropertyEditorModule>("PropertyEditor");
+			PropertyModule.RegisterCustomPropertyTypeLayout("VCamInputProfile", FOnGetPropertyTypeCustomizationInstance::CreateStatic(&FVCamInputProfileCustomization::MakeInstance));
+		};
 	
-	virtual void ShutdownModule() override
-	{
-		FPropertyEditorModule& PropertyModule = FModuleManager::GetModuleChecked<FPropertyEditorModule>("PropertyEditor");
-		PropertyModule.UnregisterCustomPropertyTypeLayout("VCamInputProfile");
-	};
+		virtual void ShutdownModule() override
+		{
+			FPropertyEditorModule& PropertyModule = FModuleManager::GetModuleChecked<FPropertyEditorModule>("PropertyEditor");
+			PropertyModule.UnregisterCustomPropertyTypeLayout("VCamInputProfile");
+		};
 
-	// Set up default nodes for VCam Core blueprint classes
-	void RegisterAutoGeneratedDefaultEvents()
-	{
-		FKismetEditorUtilities::RegisterAutoGeneratedDefaultEvent(this, UVCamBlueprintModifier::StaticClass(), "OnInitialize");
-		FKismetEditorUtilities::RegisterAutoGeneratedDefaultEvent(this, UVCamBlueprintModifier::StaticClass(), "OnDeinitialize");
-		FKismetEditorUtilities::RegisterAutoGeneratedDefaultEvent(this, UVCamBlueprintModifier::StaticClass(), "OnApply");
-	}	
-};
+		// Set up default nodes for VCam Core blueprint classes
+		void RegisterAutoGeneratedDefaultEvents()
+		{
+			FKismetEditorUtilities::RegisterAutoGeneratedDefaultEvent(this, UVCamBlueprintModifier::StaticClass(), "OnInitialize");
+			FKismetEditorUtilities::RegisterAutoGeneratedDefaultEvent(this, UVCamBlueprintModifier::StaticClass(), "OnDeinitialize");
+			FKismetEditorUtilities::RegisterAutoGeneratedDefaultEvent(this, UVCamBlueprintModifier::StaticClass(), "OnApply");
+		}	
+	};
+}
 
-IMPLEMENT_MODULE(FVCamCoreEditorModule, VCamCoreEditor);
\ No newline at end of file
+IMPLEMENT_MODULE(UE::VCamCoreEditor::Private::FVCamCoreEditorModule, VCamCoreEditor);
\ No newline at end of file
diff --git a/Engine/Plugins/Experimental/VirtualCamera/Source/VCamCoreEditor/Private/VCamInputProfileCustomization.cpp b/Engine/Plugins/Experimental/VirtualCamera/Source/VCamCoreEditor/Private/VCamInputProfileCustomization.cpp
deleted file mode 100644
diff --git a/Engine/Plugins/Experimental/VirtualCamera/Source/VCamCoreEditor/Private/VCamModifierFactory.cpp b/Engine/Plugins/Experimental/VirtualCamera/Source/VCamCoreEditor/Private/VCamModifierFactory.cpp
deleted file mode 100644
diff --git a/Engine/Plugins/Experimental/VirtualCamera/Source/VCamCoreEditor/Private/VCamObjectWithInputFactory.cpp b/Engine/Plugins/Experimental/VirtualCamera/Source/VCamCoreEditor/Private/VCamObjectWithInputFactory.cpp
deleted file mode 100644
diff --git a/Engine/Plugins/Experimental/VirtualCamera/Source/VCamCoreEditor/Private/VCamWidgetFactory.cpp b/Engine/Plugins/Experimental/VirtualCamera/Source/VCamCoreEditor/Private/VCamWidgetFactory.cpp
deleted file mode 100644
diff --git a/Engine/Plugins/Experimental/VirtualCamera/Source/VCamCoreEditor/Public/Factories/VCamModifierFactory.h b/Engine/Plugins/Experimental/VirtualCamera/Source/VCamCoreEditor/Public/Factories/VCamModifierFactory.h
new file mode 100644
--- /dev/null
+++ b/Engine/Plugins/Experimental/VirtualCamera/Source/VCamCoreEditor/Public/Factories/VCamModifierFactory.h
@@ -0,0 +1,24 @@
+// Copyright Epic Games, Inc. All Rights Reserved.
+
+#pragma once
+
+#include "CoreMinimal.h"
+#include "VCamObjectWithInputFactory.h"
+
+#include "VCamModifierFactory.generated.h"
+
+UCLASS()
+class VCAMCOREEDITOR_API UVCamModifierFactory : public UVCamObjectWithInputFactory
+{
+	GENERATED_BODY()
+public:
+	UVCamModifierFactory();
+	
+	//~ Begin UFactory Interface
+	virtual FText GetDisplayName() const override;
+	virtual FText GetToolTip() const override;
+	virtual UObject* FactoryCreateNew(UClass* Class, UObject* InParent, FName Name, EObjectFlags Flags, UObject* Context, FFeedbackContext* Warn) override;
+	virtual uint32 GetMenuCategories() const override;
+	//~ Begin UFactory Interface
+	
+};
\ No newline at end of file
diff --git a/Engine/Plugins/Experimental/VirtualCamera/Source/VCamCoreEditor/Public/Factories/VCamObjectWithInputFactory.h b/Engine/Plugins/Experimental/VirtualCamera/Source/VCamCoreEditor/Public/Factories/VCamObjectWithInputFactory.h
new file mode 100644
--- /dev/null
+++ b/Engine/Plugins/Experimental/VirtualCamera/Source/VCamCoreEditor/Public/Factories/VCamObjectWithInputFactory.h
@@ -0,0 +1,29 @@
+// Copyright Epic Games, Inc. All Rights Reserved.
+
+#pragma once
+
+#include "CoreMinimal.h"
+#include "Factories/Factory.h"
+#include "InputMappingContext.h"
+
+#include "VCamObjectWithInputFactory.generated.h"
+
+UCLASS(hidecategories = Object)
+class VCAMCOREEDITOR_API UVCamObjectWithInputFactory : public UFactory
+{
+	GENERATED_BODY()
+public:	
+	// The parent class of the created blueprint
+	UPROPERTY(EditAnywhere, Category="BlueprintVirtualSubjectFactory", meta=(AllowAbstract = "", BlueprintBaseOnly = ""))
+	TSubclassOf<UObject> ParentClass;
+
+	UPROPERTY(BlueprintReadWrite, Category = "VCam Modifier Factory")
+	TObjectPtr<UInputMappingContext> InputMappingContext;
+
+	UPROPERTY(BlueprintReadWrite, Category = "VCam Modifier Factory")
+	int32 MappingPriority;
+
+	// Begin UFactory Interface
+	virtual bool ConfigureProperties() override;
+	// End UFactory Interface
+};
\ No newline at end of file
diff --git a/Engine/Plugins/Experimental/VirtualCamera/Source/VCamCoreEditor/Public/Factories/VCamWidgetFactory.h b/Engine/Plugins/Experimental/VirtualCamera/Source/VCamCoreEditor/Public/Factories/VCamWidgetFactory.h
new file mode 100644
--- /dev/null
+++ b/Engine/Plugins/Experimental/VirtualCamera/Source/VCamCoreEditor/Public/Factories/VCamWidgetFactory.h
@@ -0,0 +1,24 @@
+// Copyright Epic Games, Inc. All Rights Reserved.
+
+#pragma once
+
+#include "CoreMinimal.h"
+#include "VCamObjectWithInputFactory.h"
+
+#include "VCamWidgetFactory.generated.h"
+
+UCLASS()
+class VCAMCOREEDITOR_API UVCamWidgetFactory : public UVCamObjectWithInputFactory
+{
+	GENERATED_BODY()
+public:
+	UVCamWidgetFactory();
+	
+	//~ Begin UFactory Interface
+	virtual FText GetDisplayName() const override;
+	virtual FText GetToolTip() const override;
+	virtual UObject* FactoryCreateNew(UClass* Class, UObject* InParent, FName Name, EObjectFlags Flags, UObject* Context, FFeedbackContext* Warn) override;
+	virtual uint32 GetMenuCategories() const override;
+	//~ Begin UFactory Interface
+	
+};
\ No newline at end of file
diff --git a/Engine/Plugins/Experimental/VirtualCamera/Source/VCamCoreEditor/Public/VCamInputProfileCustomization.h b/Engine/Plugins/Experimental/VirtualCamera/Source/VCamCoreEditor/Public/VCamInputProfileCustomization.h
deleted file mode 100644
diff --git a/Engine/Plugins/Experimental/VirtualCamera/Source/VCamCoreEditor/Public/VCamModifierFactory.h b/Engine/Plugins/Experimental/VirtualCamera/Source/VCamCoreEditor/Public/VCamModifierFactory.h
deleted file mode 100644
diff --git a/Engine/Plugins/Experimental/VirtualCamera/Source/VCamCoreEditor/Public/VCamObjectWithInputFactory.h b/Engine/Plugins/Experimental/VirtualCamera/Source/VCamCoreEditor/Public/VCamObjectWithInputFactory.h
deleted file mode 100644
diff --git a/Engine/Plugins/Experimental/VirtualCamera/Source/VCamCoreEditor/Public/VCamWidgetFactory.h b/Engine/Plugins/Experimental/VirtualCamera/Source/VCamCoreEditor/Public/VCamWidgetFactory.h
deleted file mode 100644
