Create generic tree hierarchy for Blueprints

diff --git a/Engine/Plugins/Experimental/VirtualCamera/Source/VCamCore/Private/UI/BlueprintableTreeNode.cpp b/Engine/Plugins/Experimental/VirtualCamera/Source/VCamCore/Private/UI/BlueprintableTreeNode.cpp
new file mode 100644
--- /dev/null
+++ b/Engine/Plugins/Experimental/VirtualCamera/Source/VCamCore/Private/UI/BlueprintableTreeNode.cpp
@@ -0,0 +1,93 @@
+// Copyright Epic Games, Inc. All Rights Reserved.
+
+#include "UI/BlueprintableTreeNode.h"
+
+#include "Containers/Queue.h"
+
+TArray<UBlueprintableTreeNode*> UBlueprintableTreeNode::GetChildrenByFilter(const FFilterTreeNode& FilterDelegate, bool bRecursive)
+{
+	if (!FilterDelegate.IsBound())
+	{
+		return {};
+	}
+	
+	return GetChildrenByFilter([&FilterDelegate](UBlueprintableTreeNode* Node)
+	{
+		return FilterDelegate.Execute(Node);
+	}, bRecursive);
+}
+
+TArray<UBlueprintableTreeNode*> UBlueprintableTreeNode::GetChildrenByClass(TSubclassOf<UBlueprintableTreeNode> Class, bool bRecursive)
+{
+	if (!Class)
+	{
+		return {};
+	}
+	
+	return GetChildrenByFilter([Class](UBlueprintableTreeNode* Node)
+	{
+		return Node->GetClass()->IsChildOf(Class.Get());
+	}, bRecursive);
+}
+
+void UBlueprintableTreeNode::ForEachChild(const FProcessTreeNode& ProcessDelegate, bool bOnlyFirstLevel)
+{
+	if (!ProcessDelegate.IsBound())
+	{
+		return;
+	}
+	
+	return ForEachChild([&ProcessDelegate](UBlueprintableTreeNode* Node)
+	{
+		ProcessDelegate.Execute(Node);
+	}, bOnlyFirstLevel);
+}
+
+TArray<UBlueprintableTreeNode*> UBlueprintableTreeNode::GetChildrenByFilter(TFunctionRef<bool(UBlueprintableTreeNode*)> FilterFunc, bool bRecursive)
+{
+	TArray<UBlueprintableTreeNode*> Result;
+	ForEachChild([FilterFunc, &Result](UBlueprintableTreeNode* Node)
+	{
+		if (FilterFunc(Node))
+		{
+			Result.Add(Node);
+		}
+	}, bRecursive);
+	return Result;
+}
+
+void UBlueprintableTreeNode::ForEachChild(TFunctionRef<void(UBlueprintableTreeNode*)> Func, bool bRecursive)
+{
+	TQueue<UBlueprintableTreeNode*> Queue;
+	// In a tree it is technically redundant to check for cycles but GetChildren could return anything (breaking the tree property)
+	TSet<UBlueprintableTreeNode*> EnqueuedNodes;
+
+	auto EnqueueChildren = [&Queue, &EnqueuedNodes](UBlueprintableTreeNode* Node)
+	{
+		for (const FBlueprintableTreeHierarchy& Child : Node->GetChildren())
+		{
+			if (Child.Node && !EnqueuedNodes.Contains(Child.Node))
+			{
+				Queue.Enqueue(Child.Node);
+				EnqueuedNodes.Add(Child.Node);
+			}
+		}
+	};
+	EnqueueChildren(this);
+
+	UBlueprintableTreeNode* CurrentNode;
+	while (Queue.Dequeue(CurrentNode))
+	{
+		// Blueprints could leave entries null
+		if (!CurrentNode)
+		{
+			continue;
+		}
+		
+		Func(CurrentNode);
+		if (bRecursive)
+		{
+			EnqueueChildren(CurrentNode);
+		}
+	}
+}
diff --git a/Engine/Plugins/Experimental/VirtualCamera/Source/VCamCore/Public/UI/BlueprintableTreeNode.h b/Engine/Plugins/Experimental/VirtualCamera/Source/VCamCore/Public/UI/BlueprintableTreeNode.h
new file mode 100644
--- /dev/null
+++ b/Engine/Plugins/Experimental/VirtualCamera/Source/VCamCore/Public/UI/BlueprintableTreeNode.h
@@ -0,0 +1,82 @@
+// Copyright Epic Games, Inc. All Rights Reserved.
+
+#pragma once
+
+#include "CoreMinimal.h"
+#include "UObject/Object.h"
+#include "Templates/SubclassOf.h"
+#include "BlueprintableTreeNode.generated.h"
+
+/** Allows you to create a variable with a tree structure */
+USTRUCT(BlueprintType, DisplayName = "Tree Hierarchy")
+struct VCAMCORE_API FBlueprintableTreeHierarchy
+{
+	GENERATED_BODY()
+	
+	UPROPERTY(EditAnywhere, BlueprintReadWrite, Instanced, Category = "Tree")
+	TObjectPtr<class UBlueprintableTreeNode> Node;
+};
+
+/**
+ * Allows Blueprints to create generic tree like structures that are editable in the details panel.
+ * Start by adding a struct property of type FBlueprintableTreeHierarchy as a Blueprint variable.
+ *
+ * Subclasses can define where its children come from.
+ * The easiest implementation is UBlueprintableTreeNodeWithChildList, which stores its children in an instanced array.
+ * More advanced ways could also be implemented, e.g. getting children from a data asset or another data source.
+ * 
+ * @see UBlueprintableTreeNodeWithChildList
+ */
+UCLASS(Abstract, EditInlineNew)
+class VCAMCORE_API UBlueprintableTreeNode : public UObject
+{
+	GENERATED_BODY()
+public:
+
+	/** Gets this node's children, e.g. via a data asset reference. */
+	virtual TArray<FBlueprintableTreeHierarchy> GetChildren() { unimplemented(); return {}; }
+
+	DECLARE_DYNAMIC_DELEGATE_RetVal_OneParam(bool, FFilterTreeNode, UBlueprintableTreeNode*, Node);
+	/**
+	 * Goes through Children and returns all nodes which for which the supplied FilterDelegate returns true.
+	 * @param FilterDelegate The filter to apply
+	 * @param bRecursive Whether to include children's children.
+	 */
+	UFUNCTION(BlueprintCallable, Category = "Tree")
+	TArray<UBlueprintableTreeNode*> GetChildrenByFilter(const FFilterTreeNode& FilterDelegate, bool bRecursive = false);
+
+	/** Gets all child nodes whose class is equal or descends from Class. */
+	UFUNCTION(BlueprintCallable, Category = "Tree", meta = (DisplayName = "Get Children By Class", ScriptName = "GetChildrenByClass", DeterminesOutputType = "Class"))
+	TArray<UBlueprintableTreeNode*> GetChildrenByClass(TSubclassOf<UBlueprintableTreeNode> Class, bool bRecursive = false);
+
+	/**
+	 * Goes through Children and calls the delegate on each of them.
+	 * 
+	 * @param FilterDelegate The filter to apply
+	 * @param bRecursive Whether to include children's children.
+	 */
+	DECLARE_DYNAMIC_DELEGATE_OneParam(FProcessTreeNode, UBlueprintableTreeNode*, Node);
+	UFUNCTION(BlueprintCallable, Category = "Tree")
+	void ForEachChild(const FProcessTreeNode& ProcessDelegate, bool bRecursive = false);
+	
+	TArray<UBlueprintableTreeNode*> GetChildrenByFilter(TFunctionRef<bool(UBlueprintableTreeNode*)> FilterFunc, bool bRecursive = false);
+	void ForEachChild(TFunctionRef<void(UBlueprintableTreeNode*)> Func, bool bRecursive = false);
+};
+
+/**
+ * Allows Blueprints to create generic tree like structures that are editable in the details panel.
+ * Start by adding a struct property of type FBlueprintableTreeHierarchy as a Blueprint variable.
+ * 
+ * Stores its children in an array.
+ */
+UCLASS(Abstract, Blueprintable)
+class VCAMCORE_API UBlueprintableTreeNodeWithChildList : public UBlueprintableTreeNode
+{
+	GENERATED_BODY()
+public:
+	
+	UPROPERTY(EditAnywhere, BlueprintReadWrite, Category = "Tree")
+	TArray<FBlueprintableTreeHierarchy> Children;
+
+	virtual TArray<FBlueprintableTreeHierarchy> GetChildren() override { return Children; }
+};
