Create base actor

diff --git a/Engine/Plugins/Experimental/VirtualCamera/Source/VCamCore/Private/VCamBaseActor.cpp b/Engine/Plugins/Experimental/VirtualCamera/Source/VCamCore/Private/VCamBaseActor.cpp
new file mode 100644
--- /dev/null
+++ b/Engine/Plugins/Experimental/VirtualCamera/Source/VCamCore/Private/VCamBaseActor.cpp
@@ -0,0 +1,12 @@
+// Copyright Epic Games, Inc. All Rights Reserved.
+
+#include "VCamBaseActor.h"
+
+#include "VCamComponent.h"
+
+AVCamBaseActor::AVCamBaseActor(const FObjectInitializer& ObjectInitializer)
+	: Super(ObjectInitializer)
+{
+	VCamComponent = ObjectInitializer.CreateDefaultSubobject<UVCamComponent>(this, TEXT("VCam"));
+	VCamComponent->SetupAttachment(GetCameraComponent());
+}
\ No newline at end of file
diff --git a/Engine/Plugins/Experimental/VirtualCamera/Source/VCamCore/Private/VCamComponent.cpp b/Engine/Plugins/Experimental/VirtualCamera/Source/VCamCore/Private/VCamComponent.cpp
--- a/Engine/Plugins/Experimental/VirtualCamera/Source/VCamCore/Private/VCamComponent.cpp
+++ b/Engine/Plugins/Experimental/VirtualCamera/Source/VCamCore/Private/VCamComponent.cpp
@@ -50,8 +50,12 @@ namespace VCamComponent
 
 UVCamComponent::UVCamComponent()
 {
-	// Don't run on CDO
-	if (!HasAnyFlags(RF_ClassDefaultObject) && !GIsCookerLoadingPackage)
+	// For temporary objects, the InputComponent should not be created and neither should we subscribe to global callbacks
+	// 1. The Blueprint editor has two objects:
+	//	1.1 The "real" one which saves the property data - this one is RF_ArchetypeObject
+	//	1.2 The preview one (which I assume is displayed in the viewport) - this one is RF_Transient.
+	// 2. When you drag-create an actor, level editor creates a RF_Transient template actor. After you release the mouse, a real one is created (not RF_Transient).
+	if (!HasAnyFlags(RF_ClassDefaultObject | RF_ArchetypeObject | RF_Transient) && !GIsCookerLoadingPackage)
 	{
 		// Hook into the Live Link Client for our Tick
 		IModularFeatures& ModularFeatures = IModularFeatures::Get();
@@ -81,14 +85,16 @@ UVCamComponent::UVCamComponent()
 #endif
 
 		// Setup Input
-		InputComponent = NewObject<UInputComponent>(this, UInputSettings::GetDefaultInputComponentClass(), TEXT("VCamInput0"), RF_Transient);
+		UClass* InputComponentClass = UInputSettings::GetDefaultInputComponentClass();
+		constexpr bool bIsRequired = true;
+		constexpr bool bIsTransient = true;
+		InputComponent = Cast<UInputComponent>(CreateDefaultSubobject(TEXT("VCamInput0"), InputComponentClass, InputComponentClass, bIsRequired, bIsTransient));
 
 		// Apply the Default Input profile if possible
 		if (const UVCamInputSettings* VCamInputSettings = GetDefault<UVCamInputSettings>())
 		{
 			SetInputProfileFromName(VCamInputSettings->DefaultInputProfile);
 		}
-		
 	}
 }
 
diff --git a/Engine/Plugins/Experimental/VirtualCamera/Source/VCamCore/Public/VCamBaseActor.h b/Engine/Plugins/Experimental/VirtualCamera/Source/VCamCore/Public/VCamBaseActor.h
new file mode 100644
--- /dev/null
+++ b/Engine/Plugins/Experimental/VirtualCamera/Source/VCamCore/Public/VCamBaseActor.h
@@ -0,0 +1,30 @@
+// Copyright Epic Games, Inc. All Rights Reserved.
+
+#pragma once
+
+#include "CoreMinimal.h"
+#include "CineCameraActor.h"
+#include "VCamBaseActor.generated.h"
+
+class UVCamComponent;
+
+/**
+ * Utility base class which sets up a UVCamComponent with a Cine Camera.
+ * Acts as base class for Blueprint VCamActor, which has some default properties set already. Therefore Abstract and NotPlaceable. 
+ */
+UCLASS(Abstract, Blueprintable, HideCategories = (Input, VCamComponent, SceneComponent), ShowCategories=("Input|MouseInput", "Input|TouchInput"), ComponentWrapperClass)
+class VCAMCORE_API AVCamBaseActor : public ACineCameraActor
+{
+	GENERATED_BODY()
+public:
+	
+	AVCamBaseActor(const FObjectInitializer& ObjectInitializer);
+
+	UFUNCTION(BlueprintPure, Category = VCam)
+	UVCamComponent* GetVCamComponent() const { return VCamComponent; }
+	
+private:
+	
+	UPROPERTY(VisibleAnywhere, Category = VCamBaseActor, BlueprintGetter = GetVCamComponent, meta = (DisplayAfter = "AutoActivateForPlayer", DisplayPriority = 1))
+	TObjectPtr<UVCamComponent> VCamComponent = nullptr;
+};
diff --git a/Engine/Plugins/Experimental/VirtualCamera/Source/VCamCore/Public/VCamComponent.h b/Engine/Plugins/Experimental/VirtualCamera/Source/VCamCore/Public/VCamComponent.h
--- a/Engine/Plugins/Experimental/VirtualCamera/Source/VCamCore/Public/VCamComponent.h
+++ b/Engine/Plugins/Experimental/VirtualCamera/Source/VCamCore/Public/VCamComponent.h
@@ -42,14 +42,13 @@ enum class EVCamTargetViewportID : uint8
 	Viewport4 = 4
 };
 
-UCLASS(Blueprintable, ClassGroup=(VCam), meta=(BlueprintSpawnableComponent))
+UCLASS(Blueprintable, ClassGroup = VCam, HideCategories=(Mobility), meta=(BlueprintSpawnableComponent))
 class VCAMCORE_API UVCamComponent : public USceneComponent
 {
 	GENERATED_BODY()
-
 	friend class UVCamModifier;
-
 public:
+	
 	UVCamComponent();
 
 	virtual void OnComponentDestroyed(bool bDestroyingHierarchy) override;
diff --git a/Engine/Plugins/Experimental/VirtualCamera/Source/VCamCoreEditor/Private/Customizations/VCamBaseActorCustomization.cpp b/Engine/Plugins/Experimental/VirtualCamera/Source/VCamCoreEditor/Private/Customizations/VCamBaseActorCustomization.cpp
new file mode 100644
--- /dev/null
+++ b/Engine/Plugins/Experimental/VirtualCamera/Source/VCamCoreEditor/Private/Customizations/VCamBaseActorCustomization.cpp
@@ -0,0 +1,83 @@
+// Copyright Epic Games, Inc. All Rights Reserved.
+
+#include "VCamBaseActorCustomization.h"
+
+#include "DetailCategoryBuilder.h"
+#include "DetailLayoutBuilder.h"
+#include "VCamBaseActor.h"
+#include "Algo/ForEach.h"
+#include "UObject/PropertyIterator.h"
+
+namespace UE::VCamCoreEditor::Private
+{
+	namespace
+	{
+		void SetCategoryPriorities(IDetailLayoutBuilder& DetailBuilder, UClass* Class, ECategoryPriority::Type Priority)
+		{
+			TSet<FName> CategoryNames;
+			for (TFieldIterator<FField> FieldIt(Class); FieldIt; ++FieldIt)
+			{
+				if (FieldIt->HasMetaData(TEXT("Category")))
+				{
+					CategoryNames.Add(*FieldIt->GetMetaData(TEXT("Category")));
+				}
+			}
+
+			Algo::ForEach(CategoryNames, [&DetailBuilder, Priority](FName Category)
+			{
+				DetailBuilder.EditCategory(Category, FText::GetEmpty(), Priority);
+			});
+		}
+	}
+	
+	TSharedRef<IDetailCustomization> FVCamBaseActorCustomization::MakeInstance()
+	{
+		return MakeShared<FVCamBaseActorCustomization>();
+	}
+
+	void FVCamBaseActorCustomization::CustomizeDetails(IDetailLayoutBuilder& DetailBuilder)
+	{
+		DetailBuilder.SortCategories([](const TMap<FName, IDetailCategoryBuilder*>& CategoryMap)
+		{
+			for (const TPair<FName, IDetailCategoryBuilder*>& Pair : CategoryMap)
+			{
+				int32 SortOrder = Pair.Value->GetSortOrder();
+				const FName CategoryName = Pair.Key;
+
+				if (CategoryName == "TransformCommon")
+				{
+					SortOrder = 0;
+				}
+				else if(CategoryName == "VirtualCamera")
+				{
+					SortOrder = 1;
+				}
+				else if(CategoryName == "Current Camera Settings")
+				{
+					SortOrder = 2;
+				}
+				else if(CategoryName == "CameraOptions")
+				{
+					SortOrder = 3;
+				}
+				else
+				{
+					const int32 ValueSortOrder = Pair.Value->GetSortOrder();
+					if (ValueSortOrder >= SortOrder && ValueSortOrder < SortOrder + 10)
+					{
+						SortOrder += 10;
+					}
+					else
+					{
+						continue;
+					}
+				}
+
+				Pair.Value->SetSortOrder(SortOrder);
+			}
+
+			return;
+		});
+	}
+}
+
diff --git a/Engine/Plugins/Experimental/VirtualCamera/Source/VCamCoreEditor/Private/Customizations/VCamBaseActorCustomization.h b/Engine/Plugins/Experimental/VirtualCamera/Source/VCamCoreEditor/Private/Customizations/VCamBaseActorCustomization.h
new file mode 100644
--- /dev/null
+++ b/Engine/Plugins/Experimental/VirtualCamera/Source/VCamCoreEditor/Private/Customizations/VCamBaseActorCustomization.h
@@ -0,0 +1,22 @@
+// Copyright Epic Games, Inc. All Rights Reserved.
+
+#pragma once
+
+#include "CoreMinimal.h"
+#include "IDetailCustomization.h"
+
+namespace UE::VCamCoreEditor::Private
+{
+	/** Makes the "Virtual Camera" category appear first on the actor properties. */
+	class FVCamBaseActorCustomization : public IDetailCustomization
+	{
+	public:
+
+		static TSharedRef<IDetailCustomization> MakeInstance();
+
+		//~ Begin IDetailCustomization Interface
+		virtual void CustomizeDetails(IDetailLayoutBuilder& DetailBuilder) override;
+		//~ End IDetailCustomization Interface
+	};
+}
+
diff --git a/Engine/Plugins/Experimental/VirtualCamera/Source/VCamCoreEditor/Private/VCamCoreEditorModule.cpp b/Engine/Plugins/Experimental/VirtualCamera/Source/VCamCoreEditor/Private/VCamCoreEditorModule.cpp
--- a/Engine/Plugins/Experimental/VirtualCamera/Source/VCamCoreEditor/Private/VCamCoreEditorModule.cpp
+++ b/Engine/Plugins/Experimental/VirtualCamera/Source/VCamCoreEditor/Private/VCamCoreEditorModule.cpp
@@ -7,8 +7,10 @@
 #include "Interfaces/IPluginManager.h"
 #include "Kismet2/KismetEditorUtilities.h"
 #include "PropertyEditorModule.h"
+#include "VCamBaseActor.h"
 #include "VCamComponent.h"
 #include "VCamModifier.h"
+#include "Customizations/VCamBaseActorCustomization.h"
 
 namespace UE::VCamCoreEditor::Private
 {
@@ -24,6 +26,7 @@ namespace UE::VCamCoreEditor::Private
 		
 			FPropertyEditorModule& PropertyModule = FModuleManager::GetModuleChecked<FPropertyEditorModule>("PropertyEditor");
 			PropertyModule.RegisterCustomPropertyTypeLayout("VCamInputProfile", FOnGetPropertyTypeCustomizationInstance::CreateStatic(&FVCamInputProfileCustomization::MakeInstance));
+			PropertyModule.RegisterCustomClassLayout(AVCamBaseActor::StaticClass()->GetFName(), FOnGetDetailCustomizationInstance::CreateStatic(&FVCamBaseActorCustomization::MakeInstance));
 		};
 	
 		virtual void ShutdownModule() override
