Fix input profile settings not getting saved

diff --git a/Engine/Plugins/Experimental/VirtualCamera/Config/Game.ini b/Engine/Plugins/Experimental/VirtualCamera/Config/Game.ini
deleted file mode 100644
diff --git a/Engine/Plugins/Experimental/VirtualCamera/Source/VCamInput/Public/VCamInputSettings.h b/Engine/Plugins/Experimental/VirtualCamera/Source/VCamInput/Public/VCamInputSettings.h
--- a/Engine/Plugins/Experimental/VirtualCamera/Source/VCamInput/Public/VCamInputSettings.h
+++ b/Engine/Plugins/Experimental/VirtualCamera/Source/VCamInput/Public/VCamInputSettings.h
@@ -20,28 +20,27 @@ struct VCAMINPUT_API FVCamInputProfile
 	bool operator==(const FVCamInputProfile& OtherProfile) const;
 };
 
-UCLASS(config=Game)
+UCLASS(config=Game, DefaultConfig)
 class VCAMINPUT_API UVCamInputSettings : public UDeveloperSettings
 {
 	GENERATED_BODY()
-
 public:
+	
 	UPROPERTY(EditAnywhere, Config, BlueprintReadWrite, BlueprintSetter=SetDefaultInputProfile, Category="Input", meta=(GetOptions=GetInputProfileNames))
 	FName DefaultInputProfile;
-
-	UFUNCTION(BlueprintCallable, Category="VCam Input")
-	void SetDefaultInputProfile(const FName NewDefaultInputProfile);
 	
 	UPROPERTY(EditAnywhere, Config, BlueprintReadWrite, BlueprintSetter=SetInputProfiles, Category="Input")
 	TMap<FName, FVCamInputProfile> InputProfiles;
-
+	
+	UFUNCTION(BlueprintCallable, Category="VCam")
+	static UVCamInputSettings* GetVCamInputSettings();
+	
+	UFUNCTION(BlueprintCallable, Category="VCam Input")
+	void SetDefaultInputProfile(const FName NewDefaultInputProfile);
+	
 	UFUNCTION(BlueprintCallable, Category="VCam Input")
 	void SetInputProfiles(const TMap<FName, FVCamInputProfile>& NewInputProfiles);
 
 	UFUNCTION(BlueprintCallable, Category="VCam Input")
 	TArray<FName> GetInputProfileNames() const;
-	
-private:
-	UFUNCTION(BlueprintCallable, Category="VCam")
-	static UVCamInputSettings* GetVCamInputSettings();
 };
\ No newline at end of file
