Add GetParent function to UBlueprintableTreeNode

diff --git a/Engine/Plugins/Experimental/VirtualCamera/Source/VCamCore/Public/UI/BlueprintableTreeNode.h b/Engine/Plugins/Experimental/VirtualCamera/Source/VCamCore/Public/UI/BlueprintableTreeNode.h
--- a/Engine/Plugins/Experimental/VirtualCamera/Source/VCamCore/Public/UI/BlueprintableTreeNode.h
+++ b/Engine/Plugins/Experimental/VirtualCamera/Source/VCamCore/Public/UI/BlueprintableTreeNode.h
@@ -36,6 +36,9 @@ public:
 	/** Gets this node's children, e.g. via a data asset reference. */
 	virtual TArray<FBlueprintableTreeHierarchy> GetChildren() { unimplemented(); return {}; }
 
+	UFUNCTION(BlueprintCallable, Category = "Tree")
+	virtual UBlueprintableTreeNode* GetParent() const { unimplemented(); return nullptr; }
+
 	DECLARE_DYNAMIC_DELEGATE_RetVal_OneParam(bool, FFilterTreeNode, UBlueprintableTreeNode*, Node);
 	/**
 	 * Goes through Children and returns all nodes which for which the supplied FilterDelegate returns true.
@@ -78,5 +81,8 @@ public:
 	UPROPERTY(EditAnywhere, BlueprintReadWrite, Category = "Tree")
 	TArray<FBlueprintableTreeHierarchy> Children;
 
+	//~ Begin UBlueprintableTreeNode Interface
+	virtual UBlueprintableTreeNode* GetParent() const { return CastChecked<UBlueprintableTreeNode>(GetOuter()); }
 	virtual TArray<FBlueprintableTreeHierarchy> GetChildren() override { return Children; }
+	//~ End UBlueprintableTreeNode Interface
 };
