Added type layout for FVCamConnectionTargetSettings which suggests FName entries based on the Blueprint or actor you have selected

diff --git a/Engine/Plugins/Experimental/VirtualCamera/Source/VCamCore/Private/VCamComponent.cpp b/Engine/Plugins/Experimental/VirtualCamera/Source/VCamCore/Private/VCamComponent.cpp
--- a/Engine/Plugins/Experimental/VirtualCamera/Source/VCamCore/Private/VCamComponent.cpp
+++ b/Engine/Plugins/Experimental/VirtualCamera/Source/VCamCore/Private/VCamComponent.cpp
@@ -910,6 +910,13 @@ void UVCamComponent::GetAllModifiers(TArray<UVCamModifier*>& Modifiers) const
 	}
 }
 
+TArray<FName> UVCamComponent::GetAllModifierNames() const
+{
+	TArray<FName> ModifierNames;
+	Algo::Transform(ModifierStack, ModifierNames, [](const FModifierStackEntry& StackEntry){ return StackEntry.Name; });
+	return ModifierNames;
+}
+
 UVCamModifier* UVCamComponent::GetModifierByIndex(const int32 Index) const
 {
 	if (ModifierStack.IsValidIndex(Index))
diff --git a/Engine/Plugins/Experimental/VirtualCamera/Source/VCamCore/Public/VCamComponent.h b/Engine/Plugins/Experimental/VirtualCamera/Source/VCamCore/Public/VCamComponent.h
--- a/Engine/Plugins/Experimental/VirtualCamera/Source/VCamCore/Public/VCamComponent.h
+++ b/Engine/Plugins/Experimental/VirtualCamera/Source/VCamCore/Public/VCamComponent.h
@@ -181,6 +181,10 @@ public:
 	UFUNCTION(BlueprintPure, Category = "VirtualCamera")
 	void GetAllModifiers(TArray<UVCamModifier*>& Modifiers) const;
 
+	/** Returns all the modifier names used to identifying connection points. */
+	UFUNCTION(BlueprintPure, Category = "VirtualCamera")
+	TArray<FName> GetAllModifierNames() const;
+
 	// Returns the Modifier in the Stack with the given index if it exist.
 	UFUNCTION(BlueprintPure, Category = "VirtualCamera")
 	UVCamModifier* GetModifierByIndex(const int32 Index) const;
@@ -284,23 +288,23 @@ public:
 	UPROPERTY(EditAnywhere, AdvancedDisplay, Category="VirtualCamera", meta=(ForceUnits=ms, ClampMin = "11.0"), DisplayName="Update Frequencey")
 	float UpdateFrequencyMs = 66.6f;
 
-	// Which viewport to use for this VCam
+	/** Which viewport to use for this VCam */
 	UPROPERTY(EditAnywhere, BlueprintReadWrite, Category = "VirtualCamera")
 	EVCamTargetViewportID TargetViewport = EVCamTargetViewportID::CurrentlySelected;
 
 	UPROPERTY(EditAnywhere, BlueprintReadWrite, BlueprintSetter=SetInputProfile, Category = "VirtualCamera")
 	FVCamInputProfile InputProfile;
 
-	// Call this after modifying the InputProfile in code to update the player mapped keys
+	/** List of Output Providers (executed in order) */
+	UPROPERTY(EditAnywhere, Instanced, Category="VirtualCamera")
+	TArray<TObjectPtr<UVCamOutputProviderBase>> OutputProviders;
+
+	/** Call this after modifying the InputProfile in code to update the player mapped keys */
 	void ApplyInputProfile();
 	
 	UFUNCTION(BlueprintCallable, BlueprintInternalUseOnly, Category="VirtualCamera")
 	void SetInputProfile(const FVCamInputProfile& NewInputProfile);
 
-	// List of Output Providers (executed in order)
-	UPROPERTY(EditAnywhere, Instanced, Category="VirtualCamera")
-	TArray<TObjectPtr<UVCamOutputProviderBase>> OutputProviders;
-
 	UFUNCTION(BlueprintCallable, Category="VirtualCamera")
 	void GetLiveLinkDataForCurrentFrame(FLiveLinkCameraBlueprintData& LiveLinkData);
 
@@ -332,8 +336,7 @@ public:
 	*/
 	UFUNCTION(BlueprintCallable, Category = "VirtualCamera", meta = (AutoCreateRefTerm = "Modifiers,Triggers"))
 	virtual void InjectInputForAction(const UInputAction* Action, FInputActionValue RawValue, const TArray<UInputModifier*>& Modifiers, const TArray<UInputTrigger*>& Triggers);
-
-
+	
 	/**
 	* Injects an input vector for action.
 	*/
diff --git a/Engine/Plugins/Experimental/VirtualCamera/Source/VCamCoreEditor/Private/Customizations/ConnectionTargetSettingsTypeCustomization.cpp b/Engine/Plugins/Experimental/VirtualCamera/Source/VCamCoreEditor/Private/Customizations/ConnectionTargetSettingsTypeCustomization.cpp
new file mode 100644
--- /dev/null
+++ b/Engine/Plugins/Experimental/VirtualCamera/Source/VCamCoreEditor/Private/Customizations/ConnectionTargetSettingsTypeCustomization.cpp
@@ -0,0 +1,311 @@
+// Copyright Epic Games, Inc. All Rights Reserved.
+
+#include "ConnectionTargetSettingsTypeCustomization.h"
+
+#include "VCamComponent.h"
+#include "UI/VCamConnectionStructs.h"
+
+#include "Algo/ForEach.h"
+#include "DetailWidgetRow.h"
+#include "Framework/MultiBox/MultiBoxBuilder.h"
+#include "IDetailChildrenBuilder.h"
+#include "LevelEditor.h"
+#include "Modules/ModuleManager.h"
+#include "Selection.h"
+#include "SSimpleComboButton.h"
+#include "Widgets/Images/SImage.h"
+#include "Widgets/Layout/SBox.h"
+#include "Widgets/SBoxPanel.h"
+#include "Widgets/Text/STextBlock.h"
+
+#define LOCTEXT_NAMESPACE "FConnectionTargetSettingsTypeCustomization"
+
+namespace UE::VCamCoreEditor::Private
+{
+	TSharedRef<IPropertyTypeCustomization> FConnectionTargetSettingsTypeCustomization::MakeInstance()
+	{
+		return MakeShared<FConnectionTargetSettingsTypeCustomization>();
+	}
+
+	void FConnectionTargetSettingsTypeCustomization::CustomizeHeader(
+		TSharedRef<IPropertyHandle> PropertyHandle,
+		FDetailWidgetRow& HeaderRow,
+		IPropertyTypeCustomizationUtils& CustomizationUtils)
+	{
+		HeaderRow
+			.NameContent()
+			[
+				PropertyHandle->CreatePropertyNameWidget()
+			]
+			.ValueContent()
+			[
+				PropertyHandle->CreatePropertyValueWidget()
+			];
+	}
+
+	void FConnectionTargetSettingsTypeCustomization::CustomizeChildren(
+		TSharedRef<IPropertyHandle> PropertyHandle,
+		IDetailChildrenBuilder& ChildBuilder,
+		IPropertyTypeCustomizationUtils& CustomizationUtils)
+	{
+		uint32 NumChildren;
+		PropertyHandle->GetNumChildren( NumChildren );	
+		const FName WidgetProperty = GET_MEMBER_NAME_CHECKED(FVCamConnectionTargetSettings, TargetModifierName);
+		const FName ConnectionTargetsProperty = GET_MEMBER_NAME_CHECKED(FVCamConnectionTargetSettings, TargetConnectionPoint);
+		TSharedPtr<IPropertyHandle> TargetModifierNameProperty;
+		TSharedPtr<IPropertyHandle> TargetConnectionPointProperty;
+		for (uint32 ChildIndex = 0; ChildIndex < NumChildren; ++ChildIndex)
+		{
+			TSharedPtr<IPropertyHandle> ChildProperty = PropertyHandle->GetChildHandle(ChildIndex);
+			if (ChildProperty->GetProperty() && ChildProperty->GetProperty()->GetFName() == WidgetProperty)
+			{
+				TargetModifierNameProperty = ChildProperty;
+			}
+			else if (ChildProperty->GetProperty() && ChildProperty->GetProperty()->GetFName() == ConnectionTargetsProperty)
+			{
+				TargetConnectionPointProperty = ChildProperty;
+			}
+		}
+		
+		AddScopeRow(ChildBuilder, CustomizationUtils);
+		
+		IDetailPropertyRow& ModifierRow = ChildBuilder.AddProperty(TargetModifierNameProperty.ToSharedRef());
+		CustomizeModifier(TargetModifierNameProperty.ToSharedRef(), ChildBuilder, CustomizationUtils, ModifierRow);
+
+		IDetailPropertyRow& TargetConnectionRow = ChildBuilder.AddProperty(TargetModifierNameProperty.ToSharedRef());
+		CustomizeConnectionPoint(TargetModifierNameProperty.ToSharedRef(), TargetConnectionPointProperty.ToSharedRef(), ChildBuilder, CustomizationUtils, TargetConnectionRow);
+	}
+
+	void FConnectionTargetSettingsTypeCustomization::AddScopeRow(IDetailChildrenBuilder& ChildBuilder, IPropertyTypeCustomizationUtils& CustomizationUtils)
+	{
+		// This row tells the user where the suggestions are coming from
+		ChildBuilder.AddCustomRow(FText::GetEmpty())
+			.NameContent()
+			[
+				SNew(STextBlock)
+				.Text(LOCTEXT("Scope.Name", "Scope"))
+				.ToolTipText(LOCTEXT("Scope.Tooltip", "Helps in suggesting connections points. Either:\n1. Select an Actor Blueprint containing a UVCamComponent, or\n2. Click an actor in the level containing a UVCamComponent"))
+				.Font(CustomizationUtils.GetRegularFont())
+			]
+			.ValueContent()
+			[
+				SNew(STextBlock)
+				.Text_Lambda([]()
+				{
+					const FSelectedComponentInfo ConnectionPointSourceInfo = GetUserFocusedConnectionPointSource();
+					switch (ConnectionPointSourceInfo.ComponentSource)
+					{
+					case EComponentSource::ContentBrowser:
+						check(ConnectionPointSourceInfo.Component.IsValid());
+						return FText::Format(LOCTEXT("Scope.ContentBrowser", "Asset: {0}"), FText::FromString(ConnectionPointSourceInfo.Component->GetPackage()->GetName()));
+					case EComponentSource::LevelSelection:
+						check(ConnectionPointSourceInfo.Component.IsValid());
+						return FText::Format(LOCTEXT("Scope.LevelSelection", "Actor: {0}"), FText::FromString(ConnectionPointSourceInfo.Component->GetOwner()->GetActorLabel()) );
+					case EComponentSource::None:
+					default:
+						return LOCTEXT("Scope.None", "No object selected");
+					}
+				})
+				.Font(CustomizationUtils.GetRegularFont())
+			];
+	}
+
+	void FConnectionTargetSettingsTypeCustomization::CustomizeModifier(
+		TSharedRef<IPropertyHandle> ModifierHandle,
+		IDetailChildrenBuilder& ChildBuilder,
+		IPropertyTypeCustomizationUtils& CustomizationUtils,
+		IDetailPropertyRow& Row
+		) const
+	{
+		CustomizeNameProperty(ModifierHandle, Row,
+			TAttribute<TArray<FName>>::CreateLambda([]() -> TArray<FName>
+			{
+				UVCamComponent* DataSource = GetUserFocusedConnectionPointSource().Component.Get();
+				return DataSource
+					? DataSource->GetAllModifierNames()
+					: TArray<FName>{};
+			}),
+			TAttribute<bool>::CreateLambda([](){ return GetUserFocusedConnectionPointSource().ComponentSource != EComponentSource::None; })
+			);
+	}
+
+	void FConnectionTargetSettingsTypeCustomization::CustomizeConnectionPoint(
+		TSharedRef<IPropertyHandle> ModifierHandle,
+		TSharedRef<IPropertyHandle> ConnectionPointHandle,
+		IDetailChildrenBuilder& ChildBuilder,
+		IPropertyTypeCustomizationUtils& CustomizationUtils,
+		IDetailPropertyRow& Row
+		) const
+	{
+		CustomizeNameProperty(ConnectionPointHandle, Row,
+			TAttribute<TArray<FName>>::CreateLambda([ModifierHandle]() -> TArray<FName>
+			{
+				FName ModifierName;
+				if (ModifierHandle->GetValue(ModifierName) != FPropertyAccess::Success)
+				{
+					return {};
+				}
+
+				UVCamComponent* DataSource = GetUserFocusedConnectionPointSource().Component.Get();
+				UVCamModifier* Modifier = DataSource->GetModifierByName(ModifierName);
+
+				TArray<FName> ConnectionPoints;
+				if (Modifier)
+				{
+					Modifier->ConnectionPoints.GenerateKeyArray(ConnectionPoints);
+				}
+				return ConnectionPoints;
+			}),
+			TAttribute<bool>::CreateLambda([ModifierHandle]()
+			{
+				if (GetUserFocusedConnectionPointSource().ComponentSource == EComponentSource::None)
+				{
+					return false;
+				}
+
+				FName ModifierName;
+				if (ModifierHandle->GetValue(ModifierName) != FPropertyAccess::Success)
+				{
+					return false;
+				}
+				
+				UVCamComponent* DataSource = GetUserFocusedConnectionPointSource().Component.Get();
+				return DataSource->GetModifierByName(ModifierName) != nullptr;
+			}));
+	}
+
+	FConnectionTargetSettingsTypeCustomization::FSelectedComponentInfo FConnectionTargetSettingsTypeCustomization::GetUserFocusedConnectionPointSource()
+	{
+		// Content browser
+		{
+			FEditorDelegates::LoadSelectedAssetsIfNeeded.Broadcast();
+			const UClass* const SelectedClass = GEditor->GetFirstSelectedClass(AActor::StaticClass());
+			AActor* BlueprintCDO = SelectedClass
+				? Cast<AActor>(SelectedClass->GetDefaultObject())
+				: nullptr;
+			UVCamComponent* VCamComponent = BlueprintCDO
+				? BlueprintCDO->FindComponentByClass<UVCamComponent>()
+				: nullptr;
+			if (IsValid(VCamComponent))
+			{
+				return { EComponentSource::ContentBrowser, VCamComponent };
+			}
+		}
+
+		// Level Editor
+		{
+			USelection* Selection = GEditor->GetSelectedActors();
+			UObject* SelectedObject = Selection && Selection->Num() == 1
+				? Selection->GetSelectedObject(0)
+				: nullptr;
+			AActor* SelectedActor = Cast<AActor>(SelectedObject);
+			UVCamComponent* VCamComponent = SelectedActor
+				? SelectedActor->FindComponentByClass<UVCamComponent>()
+				: nullptr;
+			if (IsValid(VCamComponent))
+			{
+				return { EComponentSource::LevelSelection, VCamComponent };
+			}
+		}
+			
+		return {};
+	}
+
+	void FConnectionTargetSettingsTypeCustomization::CustomizeNameProperty(
+		TSharedRef<IPropertyHandle> PropertyHandle,
+		IDetailPropertyRow& Row,
+		TAttribute<TArray<FName>> GetOptionsAttr,
+		TAttribute<bool> HasDataSourceAttr
+		) const
+	{
+		FLevelEditorModule& LevelEditorModule = FModuleManager::GetModuleChecked<FLevelEditorModule>(TEXT("LevelEditor"));
+		LevelEditorModule.OnActorSelectionChanged();
+		
+		Row.CustomWidget()
+			.NameContent()
+			[
+				PropertyHandle->CreatePropertyNameWidget()
+			]
+			.ValueContent()
+			[
+				SNew(SHorizontalBox)
+
+				// Warning icon if the entered value matches nothing from the list
+				+SHorizontalBox::Slot()
+				.AutoWidth()
+				[
+					SNew(SImage)
+					.ColorAndOpacity(FSlateColor::UseForeground())
+					.DesiredSizeOverride(FVector2D{ 24.f, 24.f })
+					.Image(FAppStyle::Get().GetBrush("Icons.WarningWithColor"))
+					.ToolTipText(FText::Format(LOCTEXT("InvalidValue", "Invalid value: the scope object does not contain this value for property {0}"), PropertyHandle->GetPropertyDisplayName()))
+					.Visibility_Lambda([PropertyHandle, HasDataSourceAttr, GetOptionsAttr]()
+					{
+						FName Value;
+						if (PropertyHandle->GetValue(Value) != FPropertyAccess::Success
+							// None is a "valid" value which means that the connection point should be reset
+							|| Value.IsNone())
+						{
+							return EVisibility::Collapsed;
+						}
+						return HasDataSourceAttr.Get() && !GetOptionsAttr.Get().Contains(Value)
+							? EVisibility::Visible
+							: EVisibility::Collapsed;
+					})
+				]
+
+				// Normal editing if no data source object is available
+				+SHorizontalBox::Slot()
+				.AutoWidth()
+				[
+					SNew(SBox)
+					.Visibility_Lambda([HasDataSourceAttr = HasDataSourceAttr](){ return !HasDataSourceAttr.Get() ? EVisibility::Visible : EVisibility::Collapsed; })
+					[
+						PropertyHandle->CreatePropertyValueWidget()
+					]
+				]
+
+				// Suggest data if data source object is available
+				+SHorizontalBox::Slot()
+				.AutoWidth()
+				[
+					SNew(SBox)
+					.Visibility_Lambda([HasDataSourceAttr = HasDataSourceAttr](){ return HasDataSourceAttr.Get() ? EVisibility::Visible : EVisibility::Collapsed; })
+					[
+						SNew(SComboButton)
+						.HasDownArrow(true)
+						.ComboButtonStyle(FAppStyle::Get(), "SimpleComboButton")
+						.ForegroundColor(FSlateColor::UseStyle())
+						.ButtonContent()
+						[
+							PropertyHandle->CreatePropertyValueWidget()
+						]
+						.OnGetMenuContent_Lambda([PropertyHandle, GetOptionsAttr = GetOptionsAttr]()
+						{
+							FMenuBuilder MenuBuilder(true, nullptr);
+							for (const FName& Name : GetOptionsAttr.Get())
+							{
+								MenuBuilder.AddMenuEntry(
+									FText::FromName(Name),
+									FText::GetEmpty(),
+									FSlateIcon(),
+									FUIAction(
+										FExecuteAction::CreateLambda([PropertyHandle, Name]()
+										{
+											PropertyHandle->NotifyPreChange();
+											PropertyHandle->SetValue(Name);
+											PropertyHandle->NotifyPostChange(EPropertyChangeType::ValueSet);
+										})),
+										NAME_None,
+										EUserInterfaceActionType::Button
+									);
+							}
+							return MenuBuilder.MakeWidget();
+						})
+					]
+				]
+			];
+	}
+}
+
+#undef LOCTEXT_NAMESPACE
\ No newline at end of file
diff --git a/Engine/Plugins/Experimental/VirtualCamera/Source/VCamCoreEditor/Private/Customizations/ConnectionTargetSettingsTypeCustomization.h b/Engine/Plugins/Experimental/VirtualCamera/Source/VCamCoreEditor/Private/Customizations/ConnectionTargetSettingsTypeCustomization.h
new file mode 100644
--- /dev/null
+++ b/Engine/Plugins/Experimental/VirtualCamera/Source/VCamCoreEditor/Private/Customizations/ConnectionTargetSettingsTypeCustomization.h
@@ -0,0 +1,58 @@
+// Copyright Epic Games, Inc. All Rights Reserved.
+
+#pragma once
+
+#include "CoreMinimal.h"
+#include "IPropertyTypeCustomization.h"
+
+class IDetailPropertyRow;
+class UVCamComponent;
+
+namespace UE::VCamCoreEditor::Private
+{
+	/**
+	 * If the user as selected an actor with a VCamComponent or a Blueprint with a VCamComponent in the content browser,
+	 * this will create drop-downs for the modifier and connection point properties instead of just text boxes.
+	 */
+	class FConnectionTargetSettingsTypeCustomization : public IPropertyTypeCustomization
+	{
+	public:
+		
+		static TSharedRef<IPropertyTypeCustomization> MakeInstance();
+
+		//~ Begin IPropertyTypeCustomization Interface
+		virtual void CustomizeHeader(TSharedRef<IPropertyHandle> PropertyHandle, FDetailWidgetRow& HeaderRow, IPropertyTypeCustomizationUtils& CustomizationUtils) override;
+		virtual void CustomizeChildren(TSharedRef<IPropertyHandle> PropertyHandle, IDetailChildrenBuilder& ChildBuilder, IPropertyTypeCustomizationUtils& CustomizationUtils) override;
+		//~ End IPropertyTypeCustomization Interface
+
+	private:
+		
+		void AddScopeRow(IDetailChildrenBuilder& ChildBuilder, IPropertyTypeCustomizationUtils& CustomizationUtils);
+
+		void CustomizeModifier(TSharedRef<IPropertyHandle> ModifierHandle, IDetailChildrenBuilder& ChildBuilder, IPropertyTypeCustomizationUtils& CustomizationUtils, IDetailPropertyRow& Row) const;
+		void CustomizeConnectionPoint(TSharedRef<IPropertyHandle> ModifierHandle, TSharedRef<IPropertyHandle> ConnectionPointHandle, IDetailChildrenBuilder& ChildBuilder, IPropertyTypeCustomizationUtils& CustomizationUtils, IDetailPropertyRow& Row) const;
+
+		enum class EComponentSource
+		{
+			None,
+			ContentBrowser,
+			LevelSelection
+		};
+		struct FSelectedComponentInfo
+		{
+			EComponentSource ComponentSource = EComponentSource::None;
+			TWeakObjectPtr<UVCamComponent> Component;
+		};
+		/** Gets component from 1. Blueprint class selected in content browser or 2. selected actor in world. */
+		static FSelectedComponentInfo GetUserFocusedConnectionPointSource();
+		
+		void CustomizeNameProperty(
+			TSharedRef<IPropertyHandle> PropertyHandle,
+			IDetailPropertyRow& Row,
+			TAttribute<TArray<FName>> GetOptionsAttr,
+			TAttribute<bool> HasDataSourceAttr
+			) const;
+	};
+}
+
+
diff --git a/Engine/Plugins/Experimental/VirtualCamera/Source/VCamCoreEditor/Private/VCamCoreEditorModule.cpp b/Engine/Plugins/Experimental/VirtualCamera/Source/VCamCoreEditor/Private/VCamCoreEditorModule.cpp
--- a/Engine/Plugins/Experimental/VirtualCamera/Source/VCamCoreEditor/Private/VCamCoreEditorModule.cpp
+++ b/Engine/Plugins/Experimental/VirtualCamera/Source/VCamCoreEditor/Private/VCamCoreEditorModule.cpp
@@ -3,14 +3,17 @@
 #include "Modules/ModuleInterface.h"
 #include "Modules/ModuleManager.h"
 
+#include "Customizations/ConnectionTargetSettingsTypeCustomization.h"
+#include "Customizations/VCamBaseActorCustomization.h"
 #include "Customizations/VCamInputProfileCustomization.h"
-#include "Interfaces/IPluginManager.h"
-#include "Kismet2/KismetEditorUtilities.h"
-#include "PropertyEditorModule.h"
+#include "UI/VCamConnectionStructs.h"
 #include "VCamBaseActor.h"
 #include "VCamComponent.h"
 #include "VCamModifier.h"
-#include "Customizations/VCamBaseActorCustomization.h"
+
+#include "Interfaces/IPluginManager.h"
+#include "Kismet2/KismetEditorUtilities.h"
+#include "PropertyEditorModule.h"
 
 namespace UE::VCamCoreEditor::Private
 {
@@ -25,8 +28,18 @@ namespace UE::VCamCoreEditor::Private
 			RegisterAutoGeneratedDefaultEvents();
 		
 			FPropertyEditorModule& PropertyModule = FModuleManager::GetModuleChecked<FPropertyEditorModule>("PropertyEditor");
-			PropertyModule.RegisterCustomPropertyTypeLayout("VCamInputProfile", FOnGetPropertyTypeCustomizationInstance::CreateStatic(&FVCamInputProfileCustomization::MakeInstance));
-			PropertyModule.RegisterCustomClassLayout(AVCamBaseActor::StaticClass()->GetFName(), FOnGetDetailCustomizationInstance::CreateStatic(&FVCamBaseActorCustomization::MakeInstance));
+			PropertyModule.RegisterCustomPropertyTypeLayout(
+				FVCamInputProfile::StaticStruct()->GetFName(),
+				FOnGetPropertyTypeCustomizationInstance::CreateStatic(&FVCamInputProfileCustomization::MakeInstance)
+				);
+			PropertyModule.RegisterCustomPropertyTypeLayout(
+				FVCamConnectionTargetSettings::StaticStruct()->GetFName(),
+				FOnGetPropertyTypeCustomizationInstance::CreateStatic(&FConnectionTargetSettingsTypeCustomization::MakeInstance)
+				);
+			PropertyModule.RegisterCustomClassLayout(
+				AVCamBaseActor::StaticClass()->GetFName(),
+				FOnGetDetailCustomizationInstance::CreateStatic(&FVCamBaseActorCustomization::MakeInstance)
+				);
 		};
 	
 		virtual void ShutdownModule() override
