Add switch node based on connection and connection points

diff --git a/Engine/Plugins/Experimental/VirtualCamera/Source/VCamBlueprintNodes/Private/SwitchNode/VCamK2Node_SwitchConnectionSystemBase.cpp b/Engine/Plugins/Experimental/VirtualCamera/Source/VCamBlueprintNodes/Private/SwitchNode/VCamK2Node_SwitchConnectionSystemBase.cpp
new file mode 100644
--- /dev/null
+++ b/Engine/Plugins/Experimental/VirtualCamera/Source/VCamBlueprintNodes/Private/SwitchNode/VCamK2Node_SwitchConnectionSystemBase.cpp
@@ -0,0 +1,177 @@
+// Copyright Epic Games, Inc. All Rights Reserved.
+
+#include "VCamK2Node_SwitchConnectionSystemBase.h"
+
+#include "BlueprintActionDatabaseRegistrar.h"
+#include "BlueprintNodeSpawner.h"
+#include "EdGraphSchema_K2.h"
+#include "Algo/AnyOf.h"
+#include "Algo/RemoveIf.h"
+#include "Engine/Blueprint.h"
+#include "Kismet/KismetMathLibrary.h"
+
+UVCamK2Node_SwitchConnectionSystemBase::UVCamK2Node_SwitchConnectionSystemBase()
+{
+	FunctionName = TEXT("NotEqual_NameName");
+	FunctionClass = UKismetMathLibrary::StaticClass();
+}
+
+void UVCamK2Node_SwitchConnectionSystemBase::RefreshPins()
+{
+	const TArray<FName> PinData = GetPinsToCreate();
+	RemoveOrOrphanInvalidCases(PinData);
+	AddMissingCasePins(PinData);
+
+	// The cases should have the same sort order as PinData
+	SortCasePinsToHaveSameOrderAsPinData(PinData);
+	
+	GetGraph()->NotifyGraphChanged();
+}
+
+void UVCamK2Node_SwitchConnectionSystemBase::BeginDestroy()
+{
+	Super::BeginDestroy();
+	CleanUpCallbacks();
+}
+
+void UVCamK2Node_SwitchConnectionSystemBase::PostLoad()
+{
+	Super::PostLoad();
+	SetupBlueprintModifiedCallbacks();
+}
+
+void UVCamK2Node_SwitchConnectionSystemBase::GetMenuActions(FBlueprintActionDatabaseRegistrar& ActionRegistrar) const
+{
+	UClass* NodeClass = GetClass();
+	if (ActionRegistrar.IsOpenForRegistration(NodeClass))
+	{
+		UBlueprintNodeSpawner* NodeSpawner = UBlueprintNodeSpawner::Create(NodeClass);
+		check(NodeSpawner);
+
+		NodeSpawner->CustomizeNodeDelegate = UBlueprintNodeSpawner::FCustomizeNodeDelegate::CreateLambda([](UEdGraphNode* Node, bool bIsTemplateNode)
+		{
+			if (!bIsTemplateNode)
+			{
+				UVCamK2Node_SwitchConnectionSystemBase* CastNode = Cast<UVCamK2Node_SwitchConnectionSystemBase>(Node);
+				CastNode->SetupBlueprintModifiedCallbacks();
+			}
+		});
+		
+		ActionRegistrar.AddBlueprintAction(NodeClass, NodeSpawner);
+	}
+}
+
+bool UVCamK2Node_SwitchConnectionSystemBase::IsActionFilteredOut(FBlueprintActionFilter const& Filter)
+{
+	const FBlueprintActionContext& FilterContext = Filter.Context;
+	return Algo::AnyOf(FilterContext.Blueprints, [this](UBlueprint* Blueprint){ return !SupportsBlueprintClass(Blueprint->GeneratedClass); });
+}
+
+void UVCamK2Node_SwitchConnectionSystemBase::ReconstructNode()
+{
+	Super::ReconstructNode();
+	SetupBlueprintModifiedCallbacks();
+}
+
+FEdGraphPinType UVCamK2Node_SwitchConnectionSystemBase::GetPinType() const
+{
+	FEdGraphPinType PinType;
+	PinType.PinCategory = UEdGraphSchema_K2::PC_Name;
+	return PinType;
+}
+
+void UVCamK2Node_SwitchConnectionSystemBase::CreateSelectionPin()
+{
+	UEdGraphPin* Pin = CreatePin(EGPD_Input, UEdGraphSchema_K2::PC_Name, TEXT("Selection"));
+	GetDefault<UEdGraphSchema_K2>()->SetPinAutogeneratedDefaultValueBasedOnType(Pin);
+}
+
+void UVCamK2Node_SwitchConnectionSystemBase::CreateCasePins()
+{
+	RefreshPins();
+}
+
+void UVCamK2Node_SwitchConnectionSystemBase::RemoveOrOrphanInvalidCases(const TArray<FName>& PinData)
+{
+	for (UEdGraphPin* Pin : TArray<UEdGraphPin*>(Pins))
+	{
+		const bool bIsCasePin = Pin != GetDefaultPin() && Pin->Direction == EGPD_Output;
+		const bool bIsInvalidNode = bIsCasePin && !PinData.Contains(Pin->PinName);
+		if (!bIsInvalidNode)
+		{
+			// In case it was removed the previous time it was updated
+			Pin->bOrphanedPin = false;
+			continue;
+		}
+		
+		if (Pin->HasAnyConnections())
+		{
+			Pin->bOrphanedPin = bIsInvalidNode;
+		}
+		else
+		{
+			Pin->BreakAllPinLinks();
+			UEdGraphNode::RemovePin(Pin);
+		}
+	}
+}
+
+void UVCamK2Node_SwitchConnectionSystemBase::AddMissingCasePins(const TArray<FName>& PinData)
+{
+	for (const FName& PinName : PinData)
+	{
+		const bool bContainsPin = Pins.ContainsByPredicate([PinName](UEdGraphPin* Item)
+		{
+			return Item->PinName == PinName;
+		});
+		if (!bContainsPin)
+		{
+			CreatePin(EGPD_Output, UEdGraphSchema_K2::PC_Exec, PinName);
+		}
+	}
+}
+
+void UVCamK2Node_SwitchConnectionSystemBase::SortCasePinsToHaveSameOrderAsPinData(const TArray<FName>& PinData)
+{
+	TArray<UEdGraphPin*> OnlyCaseNodes;
+	Algo::TransformIf(
+		Pins,
+		OnlyCaseNodes,
+		[this](UEdGraphPin* Node){ return Node->Direction == EGPD_Output && Node != GetDefaultPin(); },
+		[](UEdGraphPin* Node){  return Node; }
+	);
+	Pins.SetNum(Algo::RemoveIf(Pins, [&OnlyCaseNodes](UEdGraphPin* Node){ return OnlyCaseNodes.Contains(Node); }));
+	
+	OnlyCaseNodes.Sort([&PinData](const UEdGraphPin& Left, const UEdGraphPin& Right)
+	{
+		const int32 LeftIndex = PinData.IndexOfByKey(Left.PinName);
+		const int32 RightIndex = PinData.IndexOfByKey(Right.PinName);
+		// Fyi: LeftIndex and RightIndex could have INDEX_NONE value (e.g. input and default exec pins are not found).
+		return LeftIndex < RightIndex;
+	});
+	Pins.Append(OnlyCaseNodes);
+}
+
+void UVCamK2Node_SwitchConnectionSystemBase::SetupBlueprintModifiedCallbacks()
+{
+	if (!DelegateHandle.IsValid())
+	{
+		DelegateHandle = FCoreUObjectDelegates::OnObjectTransacted.AddUObject(this, &UVCamK2Node_SwitchConnectionSystemBase::OnObjectTransacted);
+	}
+}
+
+void UVCamK2Node_SwitchConnectionSystemBase::CleanUpCallbacks()
+{
+	FCoreUObjectDelegates::OnObjectTransacted.Remove(DelegateHandle);
+	DelegateHandle.Reset();
+}
+
+void UVCamK2Node_SwitchConnectionSystemBase::OnObjectTransacted(UObject* Object, const FTransactionObjectEvent&)
+{
+	UBlueprint* Blueprint = GetBlueprint();
+	const bool bModifiedBlueprintCDO = Object && Blueprint && Blueprint->GeneratedClass && Object == Blueprint->GeneratedClass->GetDefaultObject();
+	if (bModifiedBlueprintCDO)
+	{
+		RefreshPins();
+	}
+}
diff --git a/Engine/Plugins/Experimental/VirtualCamera/Source/VCamBlueprintNodes/Private/SwitchNode/VCamK2Node_SwitchConnectionSystemBase.h b/Engine/Plugins/Experimental/VirtualCamera/Source/VCamBlueprintNodes/Private/SwitchNode/VCamK2Node_SwitchConnectionSystemBase.h
new file mode 100644
--- /dev/null
+++ b/Engine/Plugins/Experimental/VirtualCamera/Source/VCamBlueprintNodes/Private/SwitchNode/VCamK2Node_SwitchConnectionSystemBase.h
@@ -0,0 +1,65 @@
+// Copyright Epic Games, Inc. All Rights Reserved.
+
+#pragma once
+
+#include "CoreMinimal.h"
+#include "K2Node_SwitchName.h"
+#include "VCamK2Node_SwitchConnectionSystemBase.generated.h"
+
+/**
+ * Base functionality for a switch node for the connections & connection points system.
+ * 
+ * It updates automatically when connections or connections points are added from a target object.
+ * Subclasses define how to obtain connections & connection points.
+ */
+UCLASS(Abstract)
+class VCAMBLUEPRINTNODES_API UVCamK2Node_SwitchConnectionSystemBase : public UK2Node_Switch
+{
+	GENERATED_BODY()
+public:
+
+	UVCamK2Node_SwitchConnectionSystemBase();
+
+	void RefreshPins();
+
+	//~ Begin UObject Interface
+	virtual void BeginDestroy() override;
+	virtual void PostLoad() override;
+	//~ End UObject Interface
+	
+	//~ Begin UK2Node Interface
+	virtual void GetMenuActions(FBlueprintActionDatabaseRegistrar& ActionRegistrar) const override;
+	virtual bool IsActionFilteredOut(FBlueprintActionFilter const& Filter) override;
+	virtual void ReconstructNode() override;
+	virtual bool CanEverInsertExecutionPin() const override { return false; }
+	virtual bool CanEverRemoveExecutionPin() const override { return false; }
+	//~ End UK2Node Interface
+
+	// UK2Node_Switch Interface
+	virtual FEdGraphPinType GetPinType() const override;
+	virtual bool SupportsAddPinButton() const override { return false; }
+	// End of UK2Node_Switch Interface
+
+protected:
+
+	virtual bool SupportsBlueprintClass(UClass* Class) const { unimplemented(); return false; }
+	virtual TArray<FName> GetPinsToCreate() const { unimplemented(); return {}; }
+
+	//~ Begin UK2Node_Switch Interface
+	virtual void CreateSelectionPin() override;
+	virtual void CreateCasePins() override;
+	//~ End UK2Node_Switch Interface
+
+private:
+	
+	FDelegateHandle DelegateHandle;
+	
+	void RemoveOrOrphanInvalidCases(const TArray<FName>& PinData);
+	void AddMissingCasePins(const TArray<FName>& PinData);
+	void SortCasePinsToHaveSameOrderAsPinData(const TArray<FName>& PinData);
+
+	void SetupBlueprintModifiedCallbacks();
+	void CleanUpCallbacks();
+	
+	void OnObjectTransacted(UObject* Object, const class FTransactionObjectEvent&);
+};
diff --git a/Engine/Plugins/Experimental/VirtualCamera/Source/VCamBlueprintNodes/Private/SwitchNode/VCamK2Node_SwitchOnModifierConnectionPoints.cpp b/Engine/Plugins/Experimental/VirtualCamera/Source/VCamBlueprintNodes/Private/SwitchNode/VCamK2Node_SwitchOnModifierConnectionPoints.cpp
new file mode 100644
--- /dev/null
+++ b/Engine/Plugins/Experimental/VirtualCamera/Source/VCamBlueprintNodes/Private/SwitchNode/VCamK2Node_SwitchOnModifierConnectionPoints.cpp
@@ -0,0 +1,52 @@
+// Copyright Epic Games, Inc. All Rights Reserved.
+
+#include "VCamK2Node_SwitchOnModifierConnectionPoints.h"
+
+#include "VCamModifier.h"
+
+#include "Engine/Blueprint.h"
+
+FText UVCamK2Node_SwitchOnModifierConnectionPoints::GetTooltipText() const
+{
+	return NSLOCTEXT("VCamEditor", "VCamK2Node.Switch_ModifierConnectionPoint.Tooltip", "Selects an output that matches the connection point");
+}
+
+FText UVCamK2Node_SwitchOnModifierConnectionPoints::GetNodeTitle(ENodeTitleType::Type TitleType) const
+{
+	return NSLOCTEXT("VCamEditor", "VCamK2Node.Switch_ModifierConnectionPoint.NodeTitle", "Switch on Connection Points");
+}
+
+bool UVCamK2Node_SwitchOnModifierConnectionPoints::SupportsBlueprintClass(UClass* Class) const
+{
+	return Class->IsChildOf(UVCamModifier::StaticClass());
+}
+
+TArray<FName> UVCamK2Node_SwitchOnModifierConnectionPoints::GetPinsToCreate() const
+{
+	TArray<FName> Result;
+	AccessBlueprintCDO([&Result](UVCamModifier* Widget)
+	{
+		Widget->ConnectionPoints.GenerateKeyArray(Result);
+	});
+	return Result;
+}
+
+void UVCamK2Node_SwitchOnModifierConnectionPoints::AccessBlueprintCDO(TFunctionRef<void(UVCamModifier*)> Func) const
+{
+	if (!GetBlueprint())
+	{
+		return;
+	}
+
+	const UBlueprint* Blueprint = GetBlueprint();
+	if (!Blueprint || Blueprint->HasAnyFlags(RF_ClassDefaultObject | RF_ArchetypeObject))
+	{
+		return;
+	}
+	
+	UObject* DefaultObject = Blueprint->GeneratedClass->GetDefaultObject();
+	if (UVCamModifier* Modifier = Cast<UVCamModifier>(DefaultObject))
+	{
+		Func(Modifier);
+	}
+}
diff --git a/Engine/Plugins/Experimental/VirtualCamera/Source/VCamBlueprintNodes/Private/SwitchNode/VCamK2Node_SwitchOnModifierConnectionPoints.h b/Engine/Plugins/Experimental/VirtualCamera/Source/VCamBlueprintNodes/Private/SwitchNode/VCamK2Node_SwitchOnModifierConnectionPoints.h
new file mode 100644
--- /dev/null
+++ b/Engine/Plugins/Experimental/VirtualCamera/Source/VCamBlueprintNodes/Private/SwitchNode/VCamK2Node_SwitchOnModifierConnectionPoints.h
@@ -0,0 +1,33 @@
+// Copyright Epic Games, Inc. All Rights Reserved.
+
+#pragma once
+
+#include "CoreMinimal.h"
+#include "VCamK2Node_SwitchConnectionSystemBase.h"
+#include "VCamK2Node_SwitchOnModifierConnectionPoints.generated.h"
+
+class UVCamModifier;
+
+/**
+ * 
+ */
+UCLASS()
+class VCAMBLUEPRINTNODES_API UVCamK2Node_SwitchOnModifierConnectionPoints : public UVCamK2Node_SwitchConnectionSystemBase
+{
+	GENERATED_BODY()
+protected:
+
+	// UEdGraphNode interface
+	virtual FText GetTooltipText() const override;
+	virtual FText GetNodeTitle(ENodeTitleType::Type TitleType) const override;
+	// End of UEdGraphNode interface
+
+	//~ Begin UVCamK2Node_SwitchConnectionSystemBase Interface
+	virtual bool SupportsBlueprintClass(UClass*) const override;
+	virtual TArray<FName> GetPinsToCreate() const override;
+	//~ End UVCamK2Node_SwitchConnectionSystemBase Interface
+
+private:
+
+	void AccessBlueprintCDO(TFunctionRef<void(UVCamModifier*)> Func) const;
+};
diff --git a/Engine/Plugins/Experimental/VirtualCamera/Source/VCamBlueprintNodes/Private/SwitchNode/VCamK2Node_SwitchOnWidgetConnections.cpp b/Engine/Plugins/Experimental/VirtualCamera/Source/VCamBlueprintNodes/Private/SwitchNode/VCamK2Node_SwitchOnWidgetConnections.cpp
new file mode 100644
--- /dev/null
+++ b/Engine/Plugins/Experimental/VirtualCamera/Source/VCamBlueprintNodes/Private/SwitchNode/VCamK2Node_SwitchOnWidgetConnections.cpp
@@ -0,0 +1,52 @@
+// Copyright Epic Games, Inc. All Rights Reserved.
+
+#include "VCamK2Node_SwitchOnWidgetConnections.h"
+
+#include "UI/VCamWidget.h"
+
+#include "Engine/Blueprint.h"
+
+FText UVCamK2Node_SwitchOnWidgetConnections::GetTooltipText() const
+{
+	return NSLOCTEXT("VCamEditor", "VCamK2Node.Switch_WidgetConnection.Tooltip", "Selects an output that matches the connection");
+}
+
+FText UVCamK2Node_SwitchOnWidgetConnections::GetNodeTitle(ENodeTitleType::Type TitleType) const
+{
+	return NSLOCTEXT("VCamEditor", "VCamK2Node.Switch_WidgetConnection.NodeTitle", "Switch on Connection");
+}
+
+bool UVCamK2Node_SwitchOnWidgetConnections::SupportsBlueprintClass(UClass* Class) const
+{
+	return Class->IsChildOf(UVCamWidget::StaticClass());
+}
+
+TArray<FName> UVCamK2Node_SwitchOnWidgetConnections::GetPinsToCreate() const
+{
+	TArray<FName> Result;
+	AccessBlueprintCDO([&Result](UVCamWidget* Widget)
+	{
+		Widget->Connections.GenerateKeyArray(Result);
+	});
+	return Result;
+}
+
+void UVCamK2Node_SwitchOnWidgetConnections::AccessBlueprintCDO(TFunctionRef<void(UVCamWidget*)> Func) const
+{
+	if (!GetBlueprint())
+	{
+		return;
+	}
+
+	const UBlueprint* Blueprint = GetBlueprint();
+	if (!Blueprint || Blueprint->HasAnyFlags(RF_ClassDefaultObject | RF_ArchetypeObject))
+	{
+		return;
+	}
+	
+	UObject* DefaultObject = Blueprint->GeneratedClass->GetDefaultObject();
+	if (UVCamWidget* Widget = Cast<UVCamWidget>(DefaultObject))
+	{
+		Func(Widget);
+	}
+}
diff --git a/Engine/Plugins/Experimental/VirtualCamera/Source/VCamBlueprintNodes/Private/SwitchNode/VCamK2Node_SwitchOnWidgetConnections.h b/Engine/Plugins/Experimental/VirtualCamera/Source/VCamBlueprintNodes/Private/SwitchNode/VCamK2Node_SwitchOnWidgetConnections.h
new file mode 100644
--- /dev/null
+++ b/Engine/Plugins/Experimental/VirtualCamera/Source/VCamBlueprintNodes/Private/SwitchNode/VCamK2Node_SwitchOnWidgetConnections.h
@@ -0,0 +1,33 @@
+// Copyright Epic Games, Inc. All Rights Reserved.
+
+#pragma once
+
+#include "CoreMinimal.h"
+#include "VCamK2Node_SwitchConnectionSystemBase.h"
+#include "VCamK2Node_SwitchOnWidgetConnections.generated.h"
+
+class UVCamWidget;
+
+/**
+ * 
+ */
+UCLASS()
+class VCAMBLUEPRINTNODES_API UVCamK2Node_SwitchOnWidgetConnections : public UVCamK2Node_SwitchConnectionSystemBase
+{
+	GENERATED_BODY()
+protected:
+
+	// UEdGraphNode interface
+	virtual FText GetTooltipText() const override;
+	virtual FText GetNodeTitle(ENodeTitleType::Type TitleType) const override;
+	// End of UEdGraphNode interface
+
+	//~ Begin UVCamK2Node_SwitchConnectionSystemBase Interface
+	virtual bool SupportsBlueprintClass(UClass*) const override;
+	virtual TArray<FName> GetPinsToCreate() const override;
+	//~ End UVCamK2Node_SwitchConnectionSystemBase Interface
+
+private:
+
+	void AccessBlueprintCDO(TFunctionRef<void(UVCamWidget*)> Func) const;
+};
diff --git a/Engine/Plugins/Experimental/VirtualCamera/Source/VCamBlueprintNodes/Private/VCamBlueprintNodes.cpp b/Engine/Plugins/Experimental/VirtualCamera/Source/VCamBlueprintNodes/Private/VCamBlueprintNodes.cpp
new file mode 100644
--- /dev/null
+++ b/Engine/Plugins/Experimental/VirtualCamera/Source/VCamBlueprintNodes/Private/VCamBlueprintNodes.cpp
@@ -0,0 +1,24 @@
+// Copyright Epic Games, Inc. All Rights Reserved.
+
+#include "Modules/ModuleInterface.h"
+#include "Modules/ModuleManager.h"
+
+namespace UE::VCamBlueprintNodes::Private
+{
+	class FVCamBlueprintNodesModule : public IModuleInterface
+	{
+	public:
+
+		virtual void StartupModule() override
+		{
+			
+		}
+	
+		virtual void ShutdownModule() override
+		{
+			
+		}
+	};
+}
+
+IMPLEMENT_MODULE(UE::VCamBlueprintNodes::Private::FVCamBlueprintNodesModule, VCamCoreEditor);
\ No newline at end of file
diff --git a/Engine/Plugins/Experimental/VirtualCamera/Source/VCamBlueprintNodes/Private/VCamBlueprintNodesLog.cpp b/Engine/Plugins/Experimental/VirtualCamera/Source/VCamBlueprintNodes/Private/VCamBlueprintNodesLog.cpp
new file mode 100644
--- /dev/null
+++ b/Engine/Plugins/Experimental/VirtualCamera/Source/VCamBlueprintNodes/Private/VCamBlueprintNodesLog.cpp
@@ -0,0 +1,5 @@
+// Copyright Epic Games, Inc. All Rights Reserved.
+
+#include "VCamBlueprintNodesLog.h"
+
+DEFINE_LOG_CATEGORY(LogVCamBlueprintNodes);
\ No newline at end of file
diff --git a/Engine/Plugins/Experimental/VirtualCamera/Source/VCamBlueprintNodes/Private/VCamBlueprintNodesLog.h b/Engine/Plugins/Experimental/VirtualCamera/Source/VCamBlueprintNodes/Private/VCamBlueprintNodesLog.h
new file mode 100644
--- /dev/null
+++ b/Engine/Plugins/Experimental/VirtualCamera/Source/VCamBlueprintNodes/Private/VCamBlueprintNodesLog.h
@@ -0,0 +1,7 @@
+// Copyright Epic Games, Inc. All Rights Reserved.
+
+#pragma once
+
+#include "CoreMinimal.h"
+
+DECLARE_LOG_CATEGORY_EXTERN(LogVCamBlueprintNodes, Log, All);
\ No newline at end of file
diff --git a/Engine/Plugins/Experimental/VirtualCamera/Source/VCamBlueprintNodes/VCamBlueprintNodes.build.cs b/Engine/Plugins/Experimental/VirtualCamera/Source/VCamBlueprintNodes/VCamBlueprintNodes.build.cs
new file mode 100644
--- /dev/null
+++ b/Engine/Plugins/Experimental/VirtualCamera/Source/VCamBlueprintNodes/VCamBlueprintNodes.build.cs
@@ -0,0 +1,32 @@
+// Copyright Epic Games, Inc. All Rights Reserved.
+
+using UnrealBuildTool;
+
+public class VCamBlueprintNodes : ModuleRules
+{
+	public VCamBlueprintNodes(ReadOnlyTargetRules Target) : base(Target)
+	{
+		PCHUsage = ModuleRules.PCHUsageMode.UseExplicitOrSharedPCHs;
+		
+		PublicDependencyModuleNames.AddRange(
+			new string[]
+			{
+				"Core",
+			}
+		);
+			
+		
+		PrivateDependencyModuleNames.AddRange(
+			new string[]
+			{
+				"BlueprintGraph",
+				"CoreUObject",
+				"Engine",
+				"InputCore",
+				"Slate",
+				"SlateCore",
+				"VCamCore",
+			}
+		);
+	}
+}
diff --git a/Engine/Plugins/Experimental/VirtualCamera/VirtualCamera.uplugin b/Engine/Plugins/Experimental/VirtualCamera/VirtualCamera.uplugin
--- a/Engine/Plugins/Experimental/VirtualCamera/VirtualCamera.uplugin
+++ b/Engine/Plugins/Experimental/VirtualCamera/VirtualCamera.uplugin
@@ -23,6 +23,20 @@
 		"Linux"
 	],
 	"Modules": [
+		{
+			"Name": "PixelStreamingVCam",
+			"Type": "Editor",
+			"LoadingPhase": "PostEngineInit",
+			"PlatformAllowList": [ "Win64", "Linux" ],
+			"TargetDenyList": [ "Server" ]
+		},
+		{
+			"Name": "VCamBlueprintNodes",
+			"Type": "UncookedOnly",
+			"LoadingPhase": "Default",
+			"PlatformAllowList": [ "Win64", "Linux" ],
+			"TargetDenyList": [ "Server" ]
+		},
 		{
 			"Name": "VCamCore",
 			"Type": "Runtime",
@@ -60,13 +74,6 @@
 			"Name": "VirtualCameraEditor",
 			"Type": "Editor",
 			"LoadingPhase": "PostEngineInit"
-		},
-		{
-			"Name": "PixelStreamingVCam",
-			"Type": "Editor",
-			"LoadingPhase": "PostEngineInit",
-			"PlatformAllowList": [ "Win64", "Linux" ],
-			"TargetDenyList": [ "Server" ]
 		}
 	],
 	"Plugins": [
