Add state switcher widget

diff --git a/Engine/Plugins/Experimental/VirtualCamera/Source/VCamCore/Private/UI/Switcher/VCamStateSwitcherWidget.cpp b/Engine/Plugins/Experimental/VirtualCamera/Source/VCamCore/Private/UI/Switcher/VCamStateSwitcherWidget.cpp
new file mode 100644
--- /dev/null
+++ b/Engine/Plugins/Experimental/VirtualCamera/Source/VCamCore/Private/UI/Switcher/VCamStateSwitcherWidget.cpp
@@ -0,0 +1,73 @@
+// Copyright Epic Games, Inc. All Rights Reserved.
+
+#include "UI/Switcher/VCamStateSwitcherWidget.h"
+
+#include "UI/Switcher/WidgetConnectionConfig.h"
+
+DEFINE_LOG_CATEGORY(LogVCamStateSwitcher);
+
+FName UVCamStateSwitcherWidget::DefaultState("Default");
+
+bool UVCamStateSwitcherWidget::SetCurrentState(FName NewState, bool bForceUpdate)
+{
+	if (NewState == CurrentState && !bForceUpdate)
+	{
+		return false;
+	}
+	
+	const FVCamWidgetConnectionState* StateConfig = States.Find(NewState);
+	if (!StateConfig)
+	{
+		UE_LOG(LogVCamStateSwitcher, Warning, TEXT("Unknown connection state %s"), *NewState.ToString());
+		return false;
+	}
+
+	const FName OldState = CurrentState;
+	OnPreStateChanged.Broadcast(this, OldState, NewState);
+	
+	for (const FWidgetConnectionConfig& WidgetConfig : StateConfig->WidgetConfigs)
+	{
+		UVCamWidget* Widget = WidgetConfig.ResolveWidget(this);
+		if (!Widget)
+		{
+			UE_CLOG(!WidgetConfig.HasNoWidgetSet(), LogVCamStateSwitcher, Warning, TEXT("Failed to find widget %s (NewState: %s)"), *WidgetConfig.Widget.ToString(), *NewState.ToString());
+			continue;
+		}
+
+		EConnectionUpdateResult UpdateResult;
+		Widget->UpdateConnectionTargets(WidgetConfig.ConnectionTargets, true, UpdateResult);
+	}
+	
+	CurrentState = NewState;
+	OnPostStateChanged.Broadcast(this, OldState, CurrentState);
+	return true;
+}
+
+#if WITH_EDITOR
+void UVCamStateSwitcherWidget::PostEditChangeProperty(FPropertyChangedEvent& PropertyChangedEvent)
+{
+	Super::PostEditChangeProperty(PropertyChangedEvent);
+
+	if (!States.Contains(DefaultState))
+	{
+		States.Add(DefaultState, {});
+	}
+
+	if (PropertyChangedEvent.MemberProperty && PropertyChangedEvent.MemberProperty->GetFName() == GET_MEMBER_NAME_CHECKED(UVCamStateSwitcherWidget, CurrentState))
+	{
+		constexpr bool bForceUpdate = true;
+		if (!SetCurrentState(CurrentState, bForceUpdate))
+		{
+			SetCurrentState(DefaultState, bForceUpdate);
+		}
+	}
+}
+#endif 
+
+void UVCamStateSwitcherWidget::NativePreConstruct()
+{
+	Super::NativePreConstruct();
+
+	constexpr bool bForceUpdate = true;
+	SetCurrentState(CurrentState, bForceUpdate);
+}
diff --git a/Engine/Plugins/Experimental/VirtualCamera/Source/VCamCore/Private/UI/Switcher/WidgetConnectionConfig.cpp b/Engine/Plugins/Experimental/VirtualCamera/Source/VCamCore/Private/UI/Switcher/WidgetConnectionConfig.cpp
new file mode 100644
--- /dev/null
+++ b/Engine/Plugins/Experimental/VirtualCamera/Source/VCamCore/Private/UI/Switcher/WidgetConnectionConfig.cpp
@@ -0,0 +1,13 @@
+// Copyright Epic Games, Inc. All Rights Reserved.
+
+#include "UI/Switcher/WidgetConnectionConfig.h"
+
+#include "Blueprint/WidgetTree.h"
+#include "UI/Switcher/VCamStateSwitcherWidget.h"
+
+UVCamWidget* FWidgetConnectionConfig::ResolveWidget(UVCamStateSwitcherWidget* OwnerWidget) const
+{
+	return !HasNoWidgetSet() && ensure(OwnerWidget->WidgetTree)
+		? OwnerWidget->WidgetTree->FindWidget<UVCamWidget>(Widget)
+		: nullptr;
+}
diff --git a/Engine/Plugins/Experimental/VirtualCamera/Source/VCamCore/Public/UI/Switcher/VCamStateSwitcherWidget.h b/Engine/Plugins/Experimental/VirtualCamera/Source/VCamCore/Public/UI/Switcher/VCamStateSwitcherWidget.h
new file mode 100644
--- /dev/null
+++ b/Engine/Plugins/Experimental/VirtualCamera/Source/VCamCore/Public/UI/Switcher/VCamStateSwitcherWidget.h
@@ -0,0 +1,67 @@
+// Copyright Epic Games, Inc. All Rights Reserved.
+
+#pragma once
+
+#include "CoreMinimal.h"
+#include "UI/VCamWidget.h"
+#include "VCamWidgetConnectionState.h"
+#include "VCamStateSwitcherWidget.generated.h"
+
+class UVCamStateSwitcherWidget;
+DECLARE_DYNAMIC_MULTICAST_DELEGATE_ThreeParams(FChangeConnectionStateEvent, UVCamStateSwitcherWidget*, Widget, FName, OldState, FName, NewState);
+
+DECLARE_LOG_CATEGORY_EXTERN(LogVCamStateSwitcher, Log, All);
+
+/**
+ * A widget that has a set of states you can switch between using SetCurrentState.
+ * A state is a collection of VCamWidgets whose connections should be rebound to new connection points.
+ */
+UCLASS()
+class VCAMCORE_API UVCamStateSwitcherWidget : public UVCamWidget
+{
+	GENERATED_BODY()
+	static FName DefaultState;
+public:
+
+	UFUNCTION(BlueprintCallable, Category = "Connections", meta = (BlueprintInternalUseOnly = "true"))
+	void K2_SetCurrentState(FName NewState) { SetCurrentState(NewState); }
+
+	/**
+	 * Switches to given state - if the state transition is valid, UpdateConnectionTargets will be called.
+	 * If CurrentState == NewState, then this call will be ignored (unless bForceUpdate == true).
+	 * 
+	 * @param NewState The new state to switch to
+	 * @param bForceUpdate Call UpdateConnectionTargets even if the CurrentState == NewState
+	 */
+	UFUNCTION(BlueprintCallable, Category = "Connections")
+	bool SetCurrentState(FName NewState, bool bForceUpdate = false);
+	
+	UFUNCTION(BlueprintPure, Category = "Connections")
+	FName GetCurrentState() const { return CurrentState; }
+
+#if WITH_EDITOR
+	virtual void PostEditChangeProperty(FPropertyChangedEvent& PropertyChangedEvent) override;
+#endif
+
+protected:
+
+	//~ Begin UUserWidget Interface
+	virtual void NativePreConstruct() override;
+	//~ End UUserWidget Interface
+	
+private:
+
+	/** Executes when the state is about to be changed */
+	UPROPERTY(BlueprintAssignable, Category = "Connections")
+	FChangeConnectionStateEvent OnPreStateChanged;
+
+	/** Executes when after the state has been changed */
+	UPROPERTY(BlueprintAssignable, Category = "Connections")
+	FChangeConnectionStateEvent OnPostStateChanged;
+	
+	UPROPERTY(EditAnywhere, Category = "Connections")
+	TMap<FName, FVCamWidgetConnectionState> States { { DefaultState, {} } };
+
+	UPROPERTY(EditAnywhere, BlueprintGetter = "GetCurrentState", BlueprintSetter = "K2_SetCurrentState", Category = "Connections")
+	FName CurrentState = DefaultState;
+};
diff --git a/Engine/Plugins/Experimental/VirtualCamera/Source/VCamCore/Public/UI/Switcher/VCamWidgetConnectionState.h b/Engine/Plugins/Experimental/VirtualCamera/Source/VCamCore/Public/UI/Switcher/VCamWidgetConnectionState.h
new file mode 100644
--- /dev/null
+++ b/Engine/Plugins/Experimental/VirtualCamera/Source/VCamCore/Public/UI/Switcher/VCamWidgetConnectionState.h
@@ -0,0 +1,19 @@
+// Copyright Epic Games, Inc. All Rights Reserved.
+
+#pragma once
+
+#include "CoreMinimal.h"
+#include "WidgetConnectionConfig.h"
+#include "VCamWidgetConnectionState.generated.h"
+
+struct FVCamConnectionTargetSettings;
+
+USTRUCT()
+struct VCAMCORE_API FVCamWidgetConnectionState
+{
+	GENERATED_BODY()
+
+	/** A list of widgets to update */
+	UPROPERTY(EditAnywhere, Category = "Connection")
+	TArray<FWidgetConnectionConfig> WidgetConfigs;
+};
\ No newline at end of file
diff --git a/Engine/Plugins/Experimental/VirtualCamera/Source/VCamCore/Public/UI/Switcher/WidgetConnectionConfig.h b/Engine/Plugins/Experimental/VirtualCamera/Source/VCamCore/Public/UI/Switcher/WidgetConnectionConfig.h
new file mode 100644
--- /dev/null
+++ b/Engine/Plugins/Experimental/VirtualCamera/Source/VCamCore/Public/UI/Switcher/WidgetConnectionConfig.h
@@ -0,0 +1,34 @@
+// Copyright Epic Games, Inc. All Rights Reserved.
+
+#pragma once
+
+#include "CoreMinimal.h"
+#include "UI/VCamConnectionStructs.h"
+#include "WidgetConnectionConfig.generated.h"
+
+class UVCamStateSwitcherWidget;
+class UVCamWidget;
+
+USTRUCT()
+struct VCAMCORE_API FWidgetConnectionConfig
+{
+	GENERATED_BODY()
+	
+	/**
+	 * Defines the widget whose connections should be updated. Must be a VCamWidget.
+	 * This is a reference to a property on the StateSwitcherVCamWidget subclass, intended to have been generated by checking
+	 * "Is Variable" in the Designer widget mode.
+	 */
+	UPROPERTY(EditAnywhere, Category = "Connection")
+	FName Widget = NAME_None;
+	
+	/**
+	 * Key: A valid key of in UVCamWidget::Connections of the widget identified via WidgetProperty.
+	 * Value: The settings to use for this connection.
+	 */
+	UPROPERTY(EditAnywhere, Category = "Connection")
+	TMap<FName, FVCamConnectionTargetSettings> ConnectionTargets;
+
+	UVCamWidget* ResolveWidget(UVCamStateSwitcherWidget* OwnerWidget) const;
+	bool HasNoWidgetSet() const { return Widget == NAME_None; }
+};
\ No newline at end of file
diff --git a/Engine/Plugins/Experimental/VirtualCamera/Source/VCamCoreEditor/Private/Customizations/StateSwitcher/ConnectionTargetNodeBuilder.cpp b/Engine/Plugins/Experimental/VirtualCamera/Source/VCamCoreEditor/Private/Customizations/StateSwitcher/ConnectionTargetNodeBuilder.cpp
new file mode 100644
--- /dev/null
+++ b/Engine/Plugins/Experimental/VirtualCamera/Source/VCamCoreEditor/Private/Customizations/StateSwitcher/ConnectionTargetNodeBuilder.cpp
@@ -0,0 +1,137 @@
+// Copyright Epic Games, Inc. All Rights Reserved.
+
+#include "ConnectionTargetNodeBuilder.h"
+
+#include "Customizations/StateSwitcher//SStringSelectionComboBox.h"
+#include "DetailWidgetRow.h"
+#include "IDetailChildrenBuilder.h"
+#include "IPropertyUtilities.h"
+#include "PropertyCustomizationHelpers.h"
+#include "Widgets/Images/SImage.h"
+
+#define LOCTEXT_NAMESPACE "FConnectionTargetNodeBuilder"
+
+namespace UE::VCamCoreEditor::Private
+{
+	FConnectionTargetNodeBuilder::FConnectionTargetNodeBuilder(
+		TSharedRef<IPropertyHandle> ConnectionTargets,
+		TAttribute<TArray<FName>> ChooseableConnections,
+		IPropertyTypeCustomizationUtils& CustomizationUtils
+		)
+		: ConnectionTargets(MoveTemp(ConnectionTargets))
+		, ChooseableConnections(MoveTemp(ChooseableConnections))
+		, RegularFont(CustomizationUtils.GetRegularFont())
+		, PropertyUtilities(CustomizationUtils.GetPropertyUtilities())
+	{
+		ConnectionTargets->SetOnPropertyValueChanged(FSimpleDelegate::CreateLambda([PropertyUtils = PropertyUtilities]()
+		{
+			PropertyUtils->ForceRefresh();
+		}));
+	}
+
+	void FConnectionTargetNodeBuilder::GenerateHeaderRowContent(FDetailWidgetRow& NodeRow)
+	{
+		NodeRow
+			.NameContent()
+			[
+				ConnectionTargets->CreatePropertyNameWidget()
+			]
+			.ValueContent()
+			[
+				ConnectionTargets->CreatePropertyValueWidget()
+			];
+	}
+
+	void FConnectionTargetNodeBuilder::GenerateChildContent(IDetailChildrenBuilder& ChildrenBuilder)
+	{
+		uint32 NumEntry = 0;
+		ConnectionTargets->GetNumChildren(NumEntry);
+		for (uint32 EntryIndex = 0; EntryIndex < NumEntry; ++EntryIndex)
+		{
+			TSharedPtr<IPropertyHandle> EntryHandle = ConnectionTargets->GetChildHandle(EntryIndex);
+			if (!EntryHandle.IsValid())
+			{
+				continue;
+			}
+
+			EntryHandle->MarkHiddenByCustomization();
+			TSharedPtr<IPropertyHandle> KeyHandle = EntryHandle->GetKeyHandle();
+			
+			constexpr bool bShowChildren = true;
+			ChildrenBuilder.AddProperty(EntryHandle.ToSharedRef())
+				.CustomWidget(bShowChildren)
+				.WholeRowContent()
+				[
+					SNew(SHorizontalBox)
+
+					// Add a warning in case the user has changed the widget from which the data of ChooseableConnections originates 
+					+SHorizontalBox::Slot()
+					.AutoWidth()
+					[
+						SNew(SImage)
+						.ColorAndOpacity(FSlateColor::UseForeground())
+						.DesiredSizeOverride(FVector2D{ 24.f, 24.f })
+						.Image(FAppStyle::Get().GetBrush("Icons.WarningWithColor"))
+						.ToolTipText(LOCTEXT("ConnectionNotFound", "This connection was not found on the target widget."))
+						.Visibility_Lambda([WeakThis = TWeakPtr<FConnectionTargetNodeBuilder>(SharedThis(this)), KeyHandle]()
+						{
+							if (TSharedPtr<FConnectionTargetNodeBuilder> ThisPin = WeakThis.Pin())
+							{
+								FName Name;
+								KeyHandle->GetValue(Name);
+								
+								return ThisPin->ChooseableConnections.Get().Contains(Name)
+									? EVisibility::Collapsed
+									: EVisibility::Visible;
+							}
+							return EVisibility::Collapsed;
+						})
+					]
+
+					+SHorizontalBox::Slot()
+					.AutoWidth()
+					[
+						SNew(SStringSelectionComboBox)
+							.SelectedItem_Lambda([KeyHandle]()
+							{
+								FString Value;
+								KeyHandle->GetValue(Value);
+								return Value;
+							})
+							.ItemList(this, &FConnectionTargetNodeBuilder::GetChooseableConnectionsAsStringArray)
+							.OnItemSelected_Lambda([KeyHandle, PropertyUtils = PropertyUtilities](const FString& SelectedItem)
+							{
+								KeyHandle->NotifyPreChange();
+								KeyHandle->SetValue(FName(*SelectedItem));
+								KeyHandle->NotifyPostChange(EPropertyChangeType::ValueSet);
+									
+								PropertyUtils->ForceRefresh();
+							})
+							.Font(RegularFont)
+					]
+
+					+SHorizontalBox::Slot()
+					.AutoWidth()
+					[
+						PropertyCustomizationHelpers::MakeInsertDeleteDuplicateButton({}, FExecuteAction::CreateLambda([EntryIndex, ConnectionTargetsPin = ConnectionTargets, PropertyUtils = PropertyUtilities]()
+						{
+							ConnectionTargetsPin->NotifyPreChange();
+							ConnectionTargetsPin->AsMap()->DeleteItem(EntryIndex);
+							ConnectionTargetsPin->NotifyPostChange(EPropertyChangeType::ValueSet);
+							
+							PropertyUtils->ForceRefresh();
+						}), {})
+					]
+				];
+		}
+	}
+
+	TArray<FString> FConnectionTargetNodeBuilder::GetChooseableConnectionsAsStringArray() const
+	{
+		TArray<FString> AsArray;
+		Algo::Transform(ChooseableConnections.Get(), AsArray, [](FName Connection){ return Connection.ToString(); });
+		return AsArray;
+	}
+}
+
+#undef LOCTEXT_NAMESPACE
\ No newline at end of file
diff --git a/Engine/Plugins/Experimental/VirtualCamera/Source/VCamCoreEditor/Private/Customizations/StateSwitcher/ConnectionTargetNodeBuilder.h b/Engine/Plugins/Experimental/VirtualCamera/Source/VCamCoreEditor/Private/Customizations/StateSwitcher/ConnectionTargetNodeBuilder.h
new file mode 100644
--- /dev/null
+++ b/Engine/Plugins/Experimental/VirtualCamera/Source/VCamCoreEditor/Private/Customizations/StateSwitcher/ConnectionTargetNodeBuilder.h
@@ -0,0 +1,48 @@
+// Copyright Epic Games, Inc. All Rights Reserved.
+
+#pragma once
+
+#include "CoreMinimal.h"
+#include "IDetailCustomNodeBuilder.h"
+#include "Fonts/SlateFontInfo.h"
+#include "Misc/Attribute.h"
+
+class IPropertyTypeCustomizationUtils;
+class IPropertyUtilities;
+
+namespace UE::VCamCoreEditor::Private
+{
+	/**
+	 * Customizes the FWidgetConnectionConfig::ConnectionTargets property.
+	 * Makes every key widget be a drop-down to the connection exposed its corresponding VCamWidget.
+	 */
+	class FConnectionTargetNodeBuilder
+		: public IDetailCustomNodeBuilder
+		, public TSharedFromThis<FConnectionTargetNodeBuilder>
+	{
+	public:
+
+		FConnectionTargetNodeBuilder(TSharedRef<IPropertyHandle> ConnectionTargets, TAttribute<TArray<FName>> ChooseableConnections, IPropertyTypeCustomizationUtils& CustomizationUtils);
+
+		//~ Begin IDetailCustomNodeBuilder Interface
+		virtual void GenerateHeaderRowContent(FDetailWidgetRow& NodeRow) override;
+		virtual void GenerateChildContent(IDetailChildrenBuilder& ChildrenBuilder) override;
+		virtual FName GetName() const override { return TEXT("Connection Targets"); }
+		//~ End IDetailCustomNodeBuilder Interface
+
+	private:
+		
+		/** Handle to FWidgetConnectionConfig::ConnectionTargets */
+		TSharedRef<IPropertyHandle> ConnectionTargets;
+
+		/** Gets the list of connections on the VCamWidget */
+		TAttribute<TArray<FName>> ChooseableConnections;
+
+		FSlateFontInfo RegularFont;
+		TSharedPtr<IPropertyUtilities> PropertyUtilities;
+
+		TArray<FString> GetChooseableConnectionsAsStringArray() const;
+	};
+}
+
+
diff --git a/Engine/Plugins/Experimental/VirtualCamera/Source/VCamCoreEditor/Private/Customizations/StateSwitcher/SStringSelectionComboBox.cpp b/Engine/Plugins/Experimental/VirtualCamera/Source/VCamCoreEditor/Private/Customizations/StateSwitcher/SStringSelectionComboBox.cpp
new file mode 100644
--- /dev/null
+++ b/Engine/Plugins/Experimental/VirtualCamera/Source/VCamCoreEditor/Private/Customizations/StateSwitcher/SStringSelectionComboBox.cpp
@@ -0,0 +1,95 @@
+// Copyright Epic Games, Inc. All Rights Reserved.
+
+#include "SStringSelectionComboBox.h"
+
+#include "Algo/Find.h"
+#include "SSearchableComboBox.h"
+#include "Algo/Compare.h"
+#include "Widgets/Text/STextBlock.h"
+
+namespace UE::VCamCoreEditor::Private
+{
+	void SStringSelectionComboBox::Construct(const FArguments& InArgs)
+	{
+		ItemListAttribute = InArgs._ItemList;
+		OnItemSelected = InArgs._OnItemSelected;
+		Font = InArgs._Font;
+
+		RefreshItemList();
+		const FString SelectedItemName = InArgs._SelectedItem.IsSet() ? InArgs._SelectedItem.Get() : FString{};
+		const TSharedPtr<FString>* InitiallySelectedItem = Algo::FindByPredicate(ItemList, [SelectedItemName](const TSharedPtr<FString>& Item)
+		{
+			return Item && *Item.Get() == SelectedItemName;
+		});
+
+		SSearchableComboBox::Construct(SSearchableComboBox::FArguments()
+			.HasDownArrow(true)
+			.OptionsSource(&ItemList)
+			.OnGenerateWidget(this, &SStringSelectionComboBox::OnGenerateComboWidget)
+			.OnSelectionChanged(this, &SStringSelectionComboBox::OnSelectionChangedInternal)
+			.OnComboBoxOpening(this, &SStringSelectionComboBox::RefreshItemList)
+			.InitiallySelectedItem(InitiallySelectedItem ? *InitiallySelectedItem : nullptr)
+			.SearchVisibility(this, &SStringSelectionComboBox::GetSearchVisibility)
+			.Content()
+			[
+				SNew(STextBlock)
+				.Text_Lambda([GetSelectedItemAttr = InArgs._SelectedItem](){ return GetSelectedItemAttr.IsSet() ? FText::FromString(GetSelectedItemAttr.Get()) : FText::GetEmpty(); })
+				.Font(Font)
+			]
+			);
+	}
+
+	void SStringSelectionComboBox::RefreshItemList()
+	{
+		TArray<FString> NewItemList = ItemListAttribute.Get();
+		const bool bHasChanged = [this, &NewItemList]()
+		{
+			if (ItemList.Num() != NewItemList.Num())
+			{
+				return true;
+			}
+
+			for (int32 i = 0; i < ItemList.Num(); ++i)
+			{
+				if (ItemList[i] && *ItemList[i].Get() != NewItemList[i])
+				{
+					return true;
+				}
+			}
+			return false;
+		}();
+		if (!bHasChanged)
+		{
+			return;
+		}
+		
+		ItemList.Empty(NewItemList.Num());
+		for (const FString& Item : NewItemList)
+		{
+			ItemList.Emplace(MakeShared<FString>(Item));
+		}
+	}
+
+	void SStringSelectionComboBox::OnSelectionChangedInternal(TSharedPtr<FString> InSelectedItem, ESelectInfo::Type SelectInfo)
+	{
+		if (InSelectedItem)
+		{
+			OnItemSelected.ExecuteIfBound(*InSelectedItem);
+		}
+	}
+
+	TSharedRef<SWidget> SStringSelectionComboBox::OnGenerateComboWidget(TSharedPtr<FString> InComboString)
+	{
+		return SNew(STextBlock)
+			.Text(FText::FromString(*InComboString))
+			.Font(Font);
+	}
+
+	EVisibility SStringSelectionComboBox::GetSearchVisibility() const
+	{
+		constexpr uint32 ShowSearchForItemCount = 5;
+		return ItemList.Num() >= ShowSearchForItemCount
+			? EVisibility::Visible
+			: EVisibility::Collapsed;
+	}
+}
diff --git a/Engine/Plugins/Experimental/VirtualCamera/Source/VCamCoreEditor/Private/Customizations/StateSwitcher/SStringSelectionComboBox.h b/Engine/Plugins/Experimental/VirtualCamera/Source/VCamCoreEditor/Private/Customizations/StateSwitcher/SStringSelectionComboBox.h
new file mode 100644
--- /dev/null
+++ b/Engine/Plugins/Experimental/VirtualCamera/Source/VCamCoreEditor/Private/Customizations/StateSwitcher/SStringSelectionComboBox.h
@@ -0,0 +1,44 @@
+// Copyright Epic Games, Inc. All Rights Reserved.
+
+#pragma once
+
+#include "CoreMinimal.h"
+#include "SSearchableComboBox.h"
+#include "SSimpleComboButton.h"
+
+namespace UE::VCamCoreEditor::Private
+{
+	class SStringSelectionComboBox : public SSearchableComboBox
+	{
+	public:
+
+		DECLARE_DELEGATE_OneParam(FOnItemSelected, const FString& /*SelectedItem*/);
+
+		SLATE_BEGIN_ARGS(SStringSelectionComboBox)
+		{}
+			SLATE_ATTRIBUTE(FString, SelectedItem)
+			SLATE_ATTRIBUTE(TArray<FString>, ItemList)
+			SLATE_EVENT(FOnItemSelected, OnItemSelected)
+			SLATE_ARGUMENT(FSlateFontInfo, Font)
+		SLATE_END_ARGS()
+
+		void Construct(const FArguments& InArgs);
+
+	private:
+
+		TAttribute<TArray<FString>> ItemListAttribute;
+		/** Wraps ItemListAttribute and is fed into SSearchableComboBox API. */
+		TArray<TSharedPtr<FString>> ItemList;
+		
+		FOnItemSelected OnItemSelected;
+		FSlateFontInfo Font;
+
+		void RefreshItemList();
+		
+		void OnSelectionChangedInternal(TSharedPtr<FString> InSelectedItem, ESelectInfo::Type SelectInfo);
+		TSharedRef<SWidget> OnGenerateComboWidget( TSharedPtr<FString> InComboString );
+		EVisibility GetSearchVisibility() const;
+	};
+}
+
+
diff --git a/Engine/Plugins/Experimental/VirtualCamera/Source/VCamCoreEditor/Private/Customizations/StateSwitcher/WidgetConnectionConfigTypeCustomization.cpp b/Engine/Plugins/Experimental/VirtualCamera/Source/VCamCoreEditor/Private/Customizations/StateSwitcher/WidgetConnectionConfigTypeCustomization.cpp
new file mode 100644
--- /dev/null
+++ b/Engine/Plugins/Experimental/VirtualCamera/Source/VCamCoreEditor/Private/Customizations/StateSwitcher/WidgetConnectionConfigTypeCustomization.cpp
@@ -0,0 +1,204 @@
+// Copyright Epic Games, Inc. All Rights Reserved.
+
+#include "WidgetConnectionConfigTypeCustomization.h"
+
+#include "ConnectionTargetNodeBuilder.h"
+#include "Customizations/StateSwitcher/SStringSelectionComboBox.h"
+#include "LogVCamEditor.h"
+#include "UI/Switcher/VCamStateSwitcherWidget.h"
+#include "UI/Switcher/WidgetConnectionConfig.h"
+#include "UI/VCamWidget.h"
+
+#include "DetailWidgetRow.h"
+#include "IDetailChildrenBuilder.h"
+#include "IPropertyUtilities.h"
+#include "Blueprint/WidgetTree.h"
+#include "UObject/PropertyIterator.h"
+
+#define LOCTEXT_NAMESPACE "FWidgetConnectionConfigCustomization"
+
+namespace UE::VCamCoreEditor::Private
+{
+	TSharedRef<IPropertyTypeCustomization> FWidgetConnectionConfigTypeCustomization::MakeInstance()
+	{
+		return MakeShared<FWidgetConnectionConfigTypeCustomization>();
+	}
+
+	void FWidgetConnectionConfigTypeCustomization::CustomizeHeader(
+		TSharedRef<IPropertyHandle> PropertyHandle,
+		FDetailWidgetRow& HeaderRow,
+		IPropertyTypeCustomizationUtils& CustomizationUtils)
+	{
+		HeaderRow.
+			NameContent()
+			[
+				PropertyHandle->CreatePropertyNameWidget()
+			]
+			.ValueContent()
+			[
+				PropertyHandle->CreatePropertyValueWidget()
+			];
+	}
+
+	void FWidgetConnectionConfigTypeCustomization::CustomizeChildren(
+		TSharedRef<IPropertyHandle> StructPropertyHandle,
+		IDetailChildrenBuilder& ChildBuilder,
+		IPropertyTypeCustomizationUtils& CustomizationUtils)
+	{
+		// Retrieve structure's child properties
+		uint32 NumChildren;
+		StructPropertyHandle->GetNumChildren( NumChildren );	
+		const FName WidgetProperty = GET_MEMBER_NAME_CHECKED(FWidgetConnectionConfig, Widget);
+		const FName ConnectionTargetsProperty = GET_MEMBER_NAME_CHECKED(FWidgetConnectionConfig, ConnectionTargets);
+
+		TSharedPtr<IPropertyHandle> WidgetReferencePropertyHandle;
+		TSharedPtr<IPropertyHandle> ConnectionTargetsPropertyHandle;
+		
+		for (uint32 ChildIndex = 0; ChildIndex < NumChildren; ++ChildIndex)
+		{
+			TSharedPtr<IPropertyHandle> ChildProperty = StructPropertyHandle->GetChildHandle(ChildIndex);
+			if (ChildProperty->GetProperty() && ChildProperty->GetProperty()->GetFName() == WidgetProperty)
+			{
+				WidgetReferencePropertyHandle = ChildProperty;
+			}
+			else if (ChildProperty->GetProperty() && ChildProperty->GetProperty()->GetFName() == ConnectionTargetsProperty)
+			{
+				ConnectionTargetsPropertyHandle = ChildProperty;
+			}
+		}
+		
+		IDetailPropertyRow& ReferenceRow = ChildBuilder.AddProperty(WidgetReferencePropertyHandle.ToSharedRef());
+		CustomizeWidgetReferenceProperty(WidgetReferencePropertyHandle.ToSharedRef(), ReferenceRow, CustomizationUtils);
+		
+		CustomizeConnectionTargetsReferenceProperty(StructPropertyHandle, ConnectionTargetsPropertyHandle.ToSharedRef(), ChildBuilder, CustomizationUtils);
+	}
+
+	void FWidgetConnectionConfigTypeCustomization::CustomizeWidgetReferenceProperty(
+		TSharedRef<IPropertyHandle> WidgetReferencePropertyHandle,
+		IDetailPropertyRow& Row,
+		IPropertyTypeCustomizationUtils& CustomizationUtils) const
+	{
+		WidgetReferencePropertyHandle->SetOnPropertyValueChanged(FSimpleDelegate::CreateLambda([PropertyUtils = CustomizationUtils.GetPropertyUtilities()]()
+		{
+			PropertyUtils->ForceRefresh();
+		}));
+		
+		Row.CustomWidget()
+			.NameContent()
+			[
+				WidgetReferencePropertyHandle->CreatePropertyNameWidget()
+			]
+			.ValueContent()
+			[
+				SNew(SStringSelectionComboBox)
+					.ToolTipText(LOCTEXT("PropertySelection.Tooltip", "Select a child widget from this Blueprint.\n\nThis can either be:\n\t- an auto-generated variable (see \"Is Variable\" check box in Designer)\n\t- a regular object property of type VCamWidget (you must make sure the property points to something valid before you change states)"))
+					.SelectedItem_Lambda([WidgetReferencePropertyHandle]()
+					{
+						FName Name;
+						WidgetReferencePropertyHandle->GetValue(Name);
+						return Name.ToString();
+					})
+					// Property list is only build once: when the Blueprint is recompiled,
+					// the details view SHOULD refresh automatically and thus all of this is reconstructed. 
+					.ItemList(GetPropertyItemList(CustomizationUtils))
+					.OnItemSelected_Lambda([WidgetReferencePropertyHandle](const FString& SelectedItem)
+					{
+						WidgetReferencePropertyHandle->SetValue(SelectedItem);
+					})
+					.Font(CustomizationUtils.GetRegularFont())
+			];
+	}
+	
+	TArray<FString> FWidgetConnectionConfigTypeCustomization::GetPropertyItemList(IPropertyTypeCustomizationUtils& CustomizationUtils) const
+	{
+		TArray<FString> PropertyItemList;
+		PropertyItemList.Add(FName(NAME_None).ToString());
+		
+		const TArray<TWeakObjectPtr<UObject>>& SelectedObjects = CustomizationUtils.GetPropertyUtilities()->GetSelectedObjects();
+		if (!ensure(SelectedObjects.Num() == 1) || !SelectedObjects[0].IsValid())
+		{
+			return PropertyItemList;
+		}
+
+		UObject* Object = SelectedObjects[0].Get();
+		UUserWidget* Widget = Cast<UUserWidget>(Object);
+		if (!Widget)
+		{
+			return {};
+		}
+		
+		Widget->WidgetTree->ForEachWidget([&PropertyItemList](UWidget* Widget)
+		{
+			const bool bIsVCamWidgetProperty = Widget && Widget->GetClass()->IsChildOf(UVCamWidget::StaticClass());
+			if (bIsVCamWidgetProperty)
+			{
+				PropertyItemList.Add(Widget->GetName());
+			}
+		});
+		return PropertyItemList;
+	}
+
+	void FWidgetConnectionConfigTypeCustomization::CustomizeConnectionTargetsReferenceProperty(
+		TSharedRef<IPropertyHandle> StructPropertyHandle,
+		TSharedRef<IPropertyHandle> ConnectionTargetsPropertyHandle,
+		IDetailChildrenBuilder& ChildBuilder,
+		IPropertyTypeCustomizationUtils& CustomizationUtils
+		) const
+	{
+		ConnectionTargetsPropertyHandle->MarkHiddenByCustomization();
+		const TSharedRef<FConnectionTargetNodeBuilder> CustomBuilder = MakeShared<FConnectionTargetNodeBuilder>(
+			ConnectionTargetsPropertyHandle,
+			CreateGetConnectionsFromChildWidgetAttribute(StructPropertyHandle, CustomizationUtils),
+			CustomizationUtils
+			);
+		ChildBuilder.AddCustomBuilder(CustomBuilder);
+	}
+
+	TAttribute<TArray<FName>> FWidgetConnectionConfigTypeCustomization::CreateGetConnectionsFromChildWidgetAttribute(
+		TSharedRef<IPropertyHandle> StructPropertyHandle,
+		IPropertyTypeCustomizationUtils& CustomizationUtils
+		) const
+	{
+		const TArray<TWeakObjectPtr<UObject>>& SelectedObjects = CustomizationUtils.GetPropertyUtilities()->GetSelectedObjects();
+		if (SelectedObjects.Num() != 1 || !SelectedObjects[0].IsValid())
+		{
+			return TArray<FName>{};
+		}
+		
+		UObject* EditedObject = SelectedObjects[0].Get();
+		UVCamStateSwitcherWidget* StateSwitcherWidget = Cast<UVCamStateSwitcherWidget>(EditedObject);
+		if (!StateSwitcherWidget)
+		{
+			UE_LOG(LogVCamEditor, Error, TEXT("FWidgetConnectionConfig was expected to be within an UVCamStateSwitcherWidget object!"));
+			return TArray<FName>{};
+		}
+
+		return TAttribute<TArray<FName>>::CreateLambda([WeakStateSwitcher = TWeakObjectPtr<UVCamStateSwitcherWidget>(StateSwitcherWidget), StructPropertyHandle]() -> TArray<FName>
+		{
+			if (!WeakStateSwitcher.IsValid())
+			{
+				return {};
+			}
+			
+			void* Data;
+			const FPropertyAccess::Result AccessResult = StructPropertyHandle->GetValueData(Data);
+			if (AccessResult != FPropertyAccess::Success)
+			{
+				return {};
+			}
+
+			FWidgetConnectionConfig* ConfigData = reinterpret_cast<FWidgetConnectionConfig*>(Data);
+			const UVCamWidget* VCamWidget = ConfigData->ResolveWidget(WeakStateSwitcher.Get());
+			if (!VCamWidget)
+			{
+				return {};
+			}
+
+			TArray<FName> Result;
+			VCamWidget->Connections.GenerateKeyArray(Result);
+			return Result;
+		});
+	}
+}
+
+#undef LOCTEXT_NAMESPACE
\ No newline at end of file
diff --git a/Engine/Plugins/Experimental/VirtualCamera/Source/VCamCoreEditor/Private/Customizations/StateSwitcher/WidgetConnectionConfigTypeCustomization.h b/Engine/Plugins/Experimental/VirtualCamera/Source/VCamCoreEditor/Private/Customizations/StateSwitcher/WidgetConnectionConfigTypeCustomization.h
new file mode 100644
--- /dev/null
+++ b/Engine/Plugins/Experimental/VirtualCamera/Source/VCamCoreEditor/Private/Customizations/StateSwitcher/WidgetConnectionConfigTypeCustomization.h
@@ -0,0 +1,33 @@
+// Copyright Epic Games, Inc. All Rights Reserved.
+
+#pragma once
+
+#include "CoreMinimal.h"
+#include "IPropertyTypeCustomization.h"
+
+class IDetailPropertyRow;
+
+namespace UE::VCamCoreEditor::Private
+{
+	/** Makes the "Virtual Camera" category appear first on the actor properties. */
+	class FWidgetConnectionConfigTypeCustomization : public IPropertyTypeCustomization
+	{
+	public:
+
+		static TSharedRef<IPropertyTypeCustomization> MakeInstance();
+
+		//~ Begin IStructCustomization Interface
+		virtual void CustomizeHeader(TSharedRef<IPropertyHandle> PropertyHandle, FDetailWidgetRow& HeaderRow, IPropertyTypeCustomizationUtils& CustomizationUtils) override;
+		virtual void CustomizeChildren(TSharedRef<IPropertyHandle> PropertyHandle, IDetailChildrenBuilder& ChildBuilder, IPropertyTypeCustomizationUtils& CustomizationUtils) override;
+		//~ End IStructCustomization Interface
+
+	private:
+
+		void CustomizeWidgetReferenceProperty(TSharedRef<IPropertyHandle> WidgetReferencePropertyHandle, IDetailPropertyRow& Row, IPropertyTypeCustomizationUtils& CustomizationUtils) const;
+		TArray<FString> GetPropertyItemList(IPropertyTypeCustomizationUtils& CustomizationUtils) const;
+
+		void CustomizeConnectionTargetsReferenceProperty(TSharedRef<IPropertyHandle> StructPropertyHandle, TSharedRef<IPropertyHandle> ConnectionTargetsPropertyHandle, IDetailChildrenBuilder& ChildBuilder, IPropertyTypeCustomizationUtils& CustomizationUtils) const;
+		TAttribute<TArray<FName>> CreateGetConnectionsFromChildWidgetAttribute(TSharedRef<IPropertyHandle> StructPropertyHandle, IPropertyTypeCustomizationUtils& CustomizationUtils) const;
+	};
+}
+
diff --git a/Engine/Plugins/Experimental/VirtualCamera/Source/VCamCoreEditor/Private/Factories/VCamStateSwitcherWidgetFactory.cpp b/Engine/Plugins/Experimental/VirtualCamera/Source/VCamCoreEditor/Private/Factories/VCamStateSwitcherWidgetFactory.cpp
new file mode 100644
--- /dev/null
+++ b/Engine/Plugins/Experimental/VirtualCamera/Source/VCamCoreEditor/Private/Factories/VCamStateSwitcherWidgetFactory.cpp
@@ -0,0 +1,24 @@
+// Copyright Epic Games, Inc. All Rights Reserved.
+
+#include "Factories/VCamStateSwitcherWidgetFactory.h"
+
+#include "UI/Switcher/VCamStateSwitcherWidget.h"
+
+#define LOCTEXT_NAMESPACE "VCamStateSwitcherWidgetFactory"
+
+UVCamStateSwitcherWidgetFactory::UVCamStateSwitcherWidgetFactory()
+{
+	ParentClass = UVCamStateSwitcherWidget::StaticClass();
+}
+
+FText UVCamStateSwitcherWidgetFactory::GetDisplayName() const
+{
+	return LOCTEXT("DisplayName", "VCam State Switcher Widget");
+}
+
+FText UVCamStateSwitcherWidgetFactory::GetToolTip() const
+{
+	return LOCTEXT("Tooltip", "Extends VCamWidget with a states.\n\nYou can switch between states using SetCurrentState.\nFor each state you can specify a set of VCam subwidgets (widgets you've added to the state switcher) and connections points those widgets should bind to in the given state.");
+}
+
+#undef LOCTEXT_NAMESPACE
diff --git a/Engine/Plugins/Experimental/VirtualCamera/Source/VCamCoreEditor/Private/LogVCamEditor.cpp b/Engine/Plugins/Experimental/VirtualCamera/Source/VCamCoreEditor/Private/LogVCamEditor.cpp
new file mode 100644
--- /dev/null
+++ b/Engine/Plugins/Experimental/VirtualCamera/Source/VCamCoreEditor/Private/LogVCamEditor.cpp
@@ -0,0 +1,5 @@
+// Copyright Epic Games, Inc. All Rights Reserved.
+
+#include "LogVCamEditor.h"
+
+DEFINE_LOG_CATEGORY(LogVCamEditor);
\ No newline at end of file
diff --git a/Engine/Plugins/Experimental/VirtualCamera/Source/VCamCoreEditor/Private/LogVCamEditor.h b/Engine/Plugins/Experimental/VirtualCamera/Source/VCamCoreEditor/Private/LogVCamEditor.h
new file mode 100644
--- /dev/null
+++ b/Engine/Plugins/Experimental/VirtualCamera/Source/VCamCoreEditor/Private/LogVCamEditor.h
@@ -0,0 +1,7 @@
+// Copyright Epic Games, Inc. All Rights Reserved.
+
+#pragma once
+
+#include "CoreMinimal.h"
+
+DECLARE_LOG_CATEGORY_EXTERN(LogVCamEditor, Log, All);
\ No newline at end of file
diff --git a/Engine/Plugins/Experimental/VirtualCamera/Source/VCamCoreEditor/Private/VCamCoreEditorModule.cpp b/Engine/Plugins/Experimental/VirtualCamera/Source/VCamCoreEditor/Private/VCamCoreEditorModule.cpp
--- a/Engine/Plugins/Experimental/VirtualCamera/Source/VCamCoreEditor/Private/VCamCoreEditorModule.cpp
+++ b/Engine/Plugins/Experimental/VirtualCamera/Source/VCamCoreEditor/Private/VCamCoreEditorModule.cpp
@@ -6,7 +6,8 @@
 #include "Customizations/ConnectionTargetSettingsTypeCustomization.h"
 #include "Customizations/VCamBaseActorCustomization.h"
 #include "Customizations/VCamInputProfileCustomization.h"
-#include "UI/VCamConnectionStructs.h"
+#include "Customizations/StateSwitcher/WidgetConnectionConfigTypeCustomization.h"
+#include "UI/Switcher/WidgetConnectionConfig.h"
 #include "VCamBaseActor.h"
 #include "VCamComponent.h"
 #include "VCamModifier.h"
@@ -28,6 +29,7 @@ namespace UE::VCamCoreEditor::Private
 			RegisterAutoGeneratedDefaultEvents();
 		
 			FPropertyEditorModule& PropertyModule = FModuleManager::GetModuleChecked<FPropertyEditorModule>("PropertyEditor");
+
 			PropertyModule.RegisterCustomPropertyTypeLayout(
 				FVCamInputProfile::StaticStruct()->GetFName(),
 				FOnGetPropertyTypeCustomizationInstance::CreateStatic(&FVCamInputProfileCustomization::MakeInstance)
@@ -36,6 +38,11 @@ namespace UE::VCamCoreEditor::Private
 				FVCamConnectionTargetSettings::StaticStruct()->GetFName(),
 				FOnGetPropertyTypeCustomizationInstance::CreateStatic(&FConnectionTargetSettingsTypeCustomization::MakeInstance)
 				);
+			PropertyModule.RegisterCustomPropertyTypeLayout(
+				FWidgetConnectionConfig::StaticStruct()->GetFName(), 
+				FOnGetPropertyTypeCustomizationInstance::CreateStatic(&FWidgetConnectionConfigTypeCustomization::MakeInstance)
+				);
+			
 			PropertyModule.RegisterCustomClassLayout(
 				AVCamBaseActor::StaticClass()->GetFName(),
 				FOnGetDetailCustomizationInstance::CreateStatic(&FVCamBaseActorCustomization::MakeInstance)
diff --git a/Engine/Plugins/Experimental/VirtualCamera/Source/VCamCoreEditor/Public/Factories/VCamStateSwitcherWidgetFactory.h b/Engine/Plugins/Experimental/VirtualCamera/Source/VCamCoreEditor/Public/Factories/VCamStateSwitcherWidgetFactory.h
new file mode 100644
--- /dev/null
+++ b/Engine/Plugins/Experimental/VirtualCamera/Source/VCamCoreEditor/Public/Factories/VCamStateSwitcherWidgetFactory.h
@@ -0,0 +1,21 @@
+// Copyright Epic Games, Inc. All Rights Reserved.
+
+#pragma once
+
+#include "CoreMinimal.h"
+#include "VCamWidgetFactory.h"
+#include "VCamStateSwitcherWidgetFactory.generated.h"
+
+UCLASS()
+class VCAMCOREEDITOR_API UVCamStateSwitcherWidgetFactory : public UVCamWidgetFactory
+{
+	GENERATED_BODY()
+public:
+	
+	UVCamStateSwitcherWidgetFactory();
+
+	//~ Begin UFactory Interface
+	virtual FText GetDisplayName() const override;
+	virtual FText GetToolTip() const override;
+	//~ End UFactory Interface
+};
\ No newline at end of file
diff --git a/Engine/Plugins/Experimental/VirtualCamera/Source/VCamCoreEditor/VCamCoreEditor.build.cs b/Engine/Plugins/Experimental/VirtualCamera/Source/VCamCoreEditor/VCamCoreEditor.build.cs
--- a/Engine/Plugins/Experimental/VirtualCamera/Source/VCamCoreEditor/VCamCoreEditor.build.cs
+++ b/Engine/Plugins/Experimental/VirtualCamera/Source/VCamCoreEditor/VCamCoreEditor.build.cs
@@ -20,19 +20,20 @@ public class VCamCoreEditor : ModuleRules
 			new string[]
 			{
 				"CoreUObject",
+				"CinematicCamera",
 				"Engine",
 				"EditorFramework",
-				"UnrealEd",
-				"Slate",
-				"SlateCore",
-				"CinematicCamera",
-				"LiveLinkInterface",
-				"VCamCore",
 				"EnhancedInput",
+				"InputCore",
+				"LiveLinkInterface",
 				"PropertyEditor",
+				"Slate",
+				"SlateCore",
+				"ToolWidgets",
 				"UMG", 
 				"UMGEditor",
-				"InputCore",
+				"UnrealEd",
+				"VCamCore",
 				"VCamInput"
 			}
 		);
