Fix CIS issue

diff --git a/Engine/Plugins/Experimental/VirtualCamera/Source/VCamCoreEditor/Private/Customizations/StateSwitcher/ConnectionTargetNodeBuilder.cpp b/Engine/Plugins/Experimental/VirtualCamera/Source/VCamCoreEditor/Private/Customizations/StateSwitcher/ConnectionTargetNodeBuilder.cpp
--- a/Engine/Plugins/Experimental/VirtualCamera/Source/VCamCoreEditor/Private/Customizations/StateSwitcher/ConnectionTargetNodeBuilder.cpp
+++ b/Engine/Plugins/Experimental/VirtualCamera/Source/VCamCoreEditor/Private/Customizations/StateSwitcher/ConnectionTargetNodeBuilder.cpp
@@ -2,7 +2,7 @@
 
 #include "ConnectionTargetNodeBuilder.h"
 
-#include "Customizations/StateSwitcher//SStringSelectionComboBox.h"
+#include "Customizations/StateSwitcher/SStringSelectionComboBox.h"
 #include "DetailWidgetRow.h"
 #include "IDetailChildrenBuilder.h"
 #include "IPropertyUtilities.h"
