General clean up:
- Create Modifier and Output folders for better code organization
- move some types into separate headers
- removed used and sort includes
- change some places to adhere to code standard (private code should be in UE::x::Private)

diff --git a/Engine/Plugins/Experimental/VirtualCamera/Source/PixelStreamingVCam/Private/VCamPixelStreamingSession.cpp b/Engine/Plugins/Experimental/VirtualCamera/Source/PixelStreamingVCam/Private/VCamPixelStreamingSession.cpp
--- a/Engine/Plugins/Experimental/VirtualCamera/Source/PixelStreamingVCam/Private/VCamPixelStreamingSession.cpp
+++ b/Engine/Plugins/Experimental/VirtualCamera/Source/PixelStreamingVCam/Private/VCamPixelStreamingSession.cpp
@@ -1,32 +1,29 @@
 // Copyright Epic Games, Inc. All Rights Reserved.
 
 #include "VCamPixelStreamingSession.h"
-#include "VCamOutputComposure.h"
-#include "PixelStreamingServers.h"
-#include "Editor/EditorPerformanceSettings.h"
-#include "Containers/UnrealString.h"
+
+#include "Output/VCamOutputComposure.h"
+
 #include "Async/Async.h"
+#include "Containers/UnrealString.h"
+#include "Editor/EditorPerformanceSettings.h"
+#include "IPixelStreamingInputHandler.h"
+#include "IPixelStreamingModule.h"
 #include "Modules/ModuleManager.h"
-#include "Math/Vector4.h"
 #include "Math/Matrix.h"
-#include "Math/TransformNonVectorized.h"
-#include "Serialization/MemoryReader.h"
-#include "VCamPixelStreamingSubsystem.h"
-#include "VCamPixelStreamingLiveLink.h"
+#include "PixelStreamingEditorModule.h"
 #include "PixelStreamingVCamLog.h"
-#include "IPixelStreamingModule.h"
 #include "PixelStreamingProtocol.h"
-#include "PixelStreamingEditorModule.h"
+#include "PixelStreamingServers.h"
+#include "Serialization/MemoryReader.h"
 #include "VCamComponent.h"
-#include "Input/HittestGrid.h"
-#include "Widgets/SVirtualWindow.h"
+#include "VCamPixelStreamingSubsystem.h"
+#include "VCamPixelStreamingLiveLink.h"
 #include "VPFullScreenUserWidget.h"
-#include "InputCoreTypes.h"
-#include "IPixelStreamingInputHandler.h"
+#include "Widgets/SVirtualWindow.h"
 
-namespace VCamPixelStreamingSession
+namespace UE::VCamPixelStreamingSession::Private
 {
-	static const FName LevelEditorName(TEXT("LevelEditor"));
 	static const FSoftClassPath EmptyUMGSoftClassPath(TEXT("/VCamCore/Assets/VCam_EmptyVisibleUMG.VCam_EmptyVisibleUMG_C"));
 } // namespace VCamPixelStreamingSession
 
@@ -72,7 +69,7 @@ void UVCamPixelStreamingSession::Activate()
 	if (UMGClass == nullptr)
 	{
 		bUsingDummyUMG = true;
-		UMGClass = VCamPixelStreamingSession::EmptyUMGSoftClassPath.TryLoadClass<UUserWidget>();
+		UMGClass = UE::VCamPixelStreamingSession::Private::EmptyUMGSoftClassPath.TryLoadClass<UUserWidget>();
 	}
 
 	if (MediaOutput == nullptr)
diff --git a/Engine/Plugins/Experimental/VirtualCamera/Source/PixelStreamingVCam/Public/VCamPixelStreamingSession.h b/Engine/Plugins/Experimental/VirtualCamera/Source/PixelStreamingVCam/Public/VCamPixelStreamingSession.h
--- a/Engine/Plugins/Experimental/VirtualCamera/Source/PixelStreamingVCam/Public/VCamPixelStreamingSession.h
+++ b/Engine/Plugins/Experimental/VirtualCamera/Source/PixelStreamingVCam/Public/VCamPixelStreamingSession.h
@@ -2,14 +2,13 @@
 
 #pragma once
 
-#include "VCamOutputProviderBase.h"
+#include "Output/VCamOutputProviderBase.h"
 #include "PixelStreamingMediaCapture.h"
 #include "PixelStreamingMediaOutput.h"
 #include "PixelStreamingServers.h"
 #include "Slate/SceneViewport.h"
-#include "Engine/TextureRenderTarget2D.h"
 #if WITH_EDITOR
-	#include "LevelEditorViewport.h"
+#include "LevelEditorViewport.h"
 #endif
 #include "VCamPixelStreamingSession.generated.h"
 
diff --git a/Engine/Plugins/Experimental/VirtualCamera/Source/VCamBlueprintNodes/Private/SwitchNode/VCamK2Node_SwitchOnModifierConnectionPoints.cpp b/Engine/Plugins/Experimental/VirtualCamera/Source/VCamBlueprintNodes/Private/SwitchNode/VCamK2Node_SwitchOnModifierConnectionPoints.cpp
--- a/Engine/Plugins/Experimental/VirtualCamera/Source/VCamBlueprintNodes/Private/SwitchNode/VCamK2Node_SwitchOnModifierConnectionPoints.cpp
+++ b/Engine/Plugins/Experimental/VirtualCamera/Source/VCamBlueprintNodes/Private/SwitchNode/VCamK2Node_SwitchOnModifierConnectionPoints.cpp
@@ -2,7 +2,7 @@
 
 #include "VCamK2Node_SwitchOnModifierConnectionPoints.h"
 
-#include "VCamModifier.h"
+#include "Modifier/VCamModifier.h"
 
 #include "Engine/Blueprint.h"
 
diff --git a/Engine/Plugins/Experimental/VirtualCamera/Source/VCamCore/Private/Modifier/VCamModifier.cpp b/Engine/Plugins/Experimental/VirtualCamera/Source/VCamCore/Private/Modifier/VCamModifier.cpp
new file mode 100644
--- /dev/null
+++ b/Engine/Plugins/Experimental/VirtualCamera/Source/VCamCore/Private/Modifier/VCamModifier.cpp
@@ -0,0 +1,163 @@
+// Copyright Epic Games, Inc. All Rights Reserved.
+
+#include "Modifier/VCamModifier.h"
+
+#include "Modifier/ModifierStackEntry.h"
+#include "VCamComponent.h"
+
+#include "EnhancedInputComponent.h"
+#include "Engine/InputDelegateBinding.h"
+
+void UVCamBlueprintModifier::Initialize(UVCamModifierContext* Context, UInputComponent* InputComponent)
+{
+	// Forward the Initialize call to the Blueprint Event
+	{
+		FEditorScriptExecutionGuard ScriptGuard;
+
+		OnInitialize(Context);
+	}
+
+	UVCamModifier::Initialize(Context, InputComponent);
+}
+
+void UVCamBlueprintModifier::Deinitialize()
+{
+	// Forward the Initialize call to the Blueprint Event only if we have a valid world
+	if (GetWorld())
+	{
+		FEditorScriptExecutionGuard ScriptGuard;
+
+		OnDeinitialize();
+	}
+	
+	Super::Deinitialize();
+}
+
+void UVCamBlueprintModifier::Apply(UVCamModifierContext* Context, UCineCameraComponent* CameraComponent, const float DeltaTime)
+{
+	// Forward the Apply call to the Blueprint Event
+	{
+		FEditorScriptExecutionGuard ScriptGuard;
+
+		OnApply(Context, CameraComponent, DeltaTime);
+	}
+}
+
+void UVCamModifier::Initialize(UVCamModifierContext* Context, UInputComponent* InputComponent /* = nullptr */)
+{
+	// Binds any dynamic input delegates to the provided input component
+	if (IsValid(InputComponent))
+	{
+		UInputDelegateBinding::BindInputDelegates(GetClass(), InputComponent, this);
+	}
+
+	bRequiresInitialization = false;
+}
+
+void UVCamModifier::Deinitialize()
+{
+	UVCamComponent* VCamComponent = GetOwningVCamComponent();
+	if (IsValid(VCamComponent))
+	{
+		VCamComponent->UnregisterObjectForInput(this);
+	}
+
+	bRequiresInitialization = true;
+}
+
+void UVCamModifier::PostLoad()
+{
+	Super::PostLoad();
+
+	bRequiresInitialization = true;
+}
+
+UVCamComponent* UVCamModifier::GetOwningVCamComponent() const
+{
+	return GetTypedOuter<UVCamComponent>();
+}
+
+void UVCamModifier::GetCurrentLiveLinkDataFromOwningComponent(FLiveLinkCameraBlueprintData& LiveLinkData)
+{
+	if (UVCamComponent* OwningComponent = GetOwningVCamComponent())
+	{
+		OwningComponent->GetLiveLinkDataForCurrentFrame(LiveLinkData);
+	}
+}
+
+void UVCamModifier::SetEnabled(bool bNewEnabled)
+{
+	if (FModifierStackEntry* StackEntry = GetCorrespondingStackEntry())
+	{
+		StackEntry->bEnabled = bNewEnabled;
+	}
+
+	if (!bNewEnabled)
+	{
+		Deinitialize();
+	}
+}
+
+bool UVCamModifier::IsEnabled() const
+{
+	if (FModifierStackEntry* StackEntry = GetCorrespondingStackEntry())
+	{
+		return StackEntry->bEnabled;
+	}
+
+	return false;
+}
+
+bool UVCamModifier::SetStackEntryName(FName NewName)
+{
+	if (const UVCamComponent* ParentComponent = GetOwningVCamComponent())
+	{
+		const bool bNameAlreadyExists = ParentComponent->ModifierStack.ContainsByPredicate([NewName](const FModifierStackEntry& StackEntryToTest)
+		{
+			return StackEntryToTest.Name.IsEqual(NewName);
+		});
+
+		if (!bNameAlreadyExists)
+		{
+			if (FModifierStackEntry* StackEntry = GetCorrespondingStackEntry())
+			{
+				StackEntry->Name = NewName;
+				return true;
+			}
+		}
+	}
+	return false;
+}
+
+FName UVCamModifier::GetStackEntryName() const
+{
+	if (FModifierStackEntry* StackEntry = GetCorrespondingStackEntry())
+	{
+		return StackEntry->Name;
+	}
+	return NAME_None;
+}
+
+FModifierStackEntry* UVCamModifier::GetCorrespondingStackEntry() const
+{
+	FModifierStackEntry* StackEntry = nullptr;
+
+	if (UVCamComponent* ParentComponent = GetOwningVCamComponent())
+	{
+		StackEntry = ParentComponent->ModifierStack.FindByPredicate([this](const FModifierStackEntry& StackEntryToTest) 
+		{
+			return StackEntryToTest.GeneratedModifier == this;
+		});
+	}
+
+	return StackEntry;
+}
+
+UWorld* UVCamModifier::GetWorld() const
+{
+	if (UVCamComponent* ParentComponent = GetOwningVCamComponent())
+	{
+		return ParentComponent->GetWorld();
+	}
+	return nullptr;
+}
\ No newline at end of file
diff --git a/Engine/Plugins/Experimental/VirtualCamera/Source/VCamCore/Private/Modifier/VCamModifierInterface.cpp b/Engine/Plugins/Experimental/VirtualCamera/Source/VCamCore/Private/Modifier/VCamModifierInterface.cpp
new file mode 100644
--- /dev/null
+++ b/Engine/Plugins/Experimental/VirtualCamera/Source/VCamCore/Private/Modifier/VCamModifierInterface.cpp
@@ -0,0 +1,11 @@
+// Copyright Epic Games, Inc. All Rights Reserved.
+
+#include "Modifier/VCamModifierInterface.h"
+#include "VCamComponent.h"
+
+// Default fallback implementation
+void IVCamModifierInterface::OnVCamComponentChanged_Implementation(UVCamComponent* VCam)
+{
+	UE_LOG(LogVCamComponent, Warning, TEXT("The default function for On VCam Component Changed has been called. New component is: %s"), *VCam->GetName());
+}
+
diff --git a/Engine/Plugins/Experimental/VirtualCamera/Source/VCamCore/Private/Output/VCamOutputComposure.cpp b/Engine/Plugins/Experimental/VirtualCamera/Source/VCamCore/Private/Output/VCamOutputComposure.cpp
new file mode 100644
--- /dev/null
+++ b/Engine/Plugins/Experimental/VirtualCamera/Source/VCamCore/Private/Output/VCamOutputComposure.cpp
@@ -0,0 +1,51 @@
+// Copyright Epic Games, Inc. All Rights Reserved.
+
+#include "Output/VCamOutputComposure.h"
+
+void UVCamOutputComposure::CreateUMG()
+{
+	DisplayType = EVPWidgetDisplayType::Composure;
+
+	// Super must be here so that the UMGWidget is already created
+	Super::CreateUMG();
+
+	if (UMGWidget)
+	{
+		UMGWidget->PostProcessDisplayType.ComposureLayerTargets.Empty();
+
+		for (TSoftObjectPtr<ACompositingElement> Layer : LayerTargets)
+		{
+			if (Layer.IsValid())
+			{
+				UMGWidget->PostProcessDisplayType.ComposureLayerTargets.Emplace(Layer.Get());
+			}
+		}
+	}
+	else
+	{
+		UE_LOG(LogVCamOutputProvider, Error, TEXT("Composure mode - Either UMG Class is None or Super::CreateUMG isn't being called first!"));
+	}
+}
+
+#if WITH_EDITOR
+void UVCamOutputComposure::PostEditChangeProperty(FPropertyChangedEvent& PropertyChangedEvent)
+{
+	FProperty* Property = PropertyChangedEvent.MemberProperty;
+
+	if (Property && PropertyChangedEvent.ChangeType != EPropertyChangeType::Interactive)
+	{
+		static FName NAME_LayerTargets = GET_MEMBER_NAME_CHECKED(UVCamOutputComposure, LayerTargets);
+
+		if (Property->GetFName() == NAME_LayerTargets)
+		{
+			if (bIsActive)
+			{
+				SetActive(false);
+				SetActive(true);
+			}
+		}
+	}
+
+	Super::PostEditChangeProperty(PropertyChangedEvent);
+}
+#endif
diff --git a/Engine/Plugins/Experimental/VirtualCamera/Source/VCamCore/Private/Output/VCamOutputMediaOutput.cpp b/Engine/Plugins/Experimental/VirtualCamera/Source/VCamCore/Private/Output/VCamOutputMediaOutput.cpp
new file mode 100644
--- /dev/null
+++ b/Engine/Plugins/Experimental/VirtualCamera/Source/VCamCore/Private/Output/VCamOutputMediaOutput.cpp
@@ -0,0 +1,100 @@
+// Copyright Epic Games, Inc. All Rights Reserved.
+
+#include "Output/VCamOutputMediaOutput.h"
+
+#include "Output/VCamOutputComposure.h"
+
+void UVCamOutputMediaOutput::Activate()
+{
+	StartCapturing();
+
+	Super::Activate();
+}
+
+void UVCamOutputMediaOutput::Deactivate()
+{
+	StopCapturing();
+
+	Super::Deactivate();
+}
+
+void UVCamOutputMediaOutput::StartCapturing()
+{
+	if (OutputConfig)
+	{
+		MediaCapture = OutputConfig->CreateMediaCapture();
+		if (MediaCapture)
+		{
+			FMediaCaptureOptions Options;
+			Options.bResizeSourceBuffer = true;
+
+			// If we are rendering from a ComposureOutputProvider, get the requested render target and use that instead of the viewport
+			if (UVCamOutputComposure* ComposureProvider = Cast<UVCamOutputComposure>(GetOtherOutputProviderByIndex(FromComposureOutputProviderIndex)))
+			{
+				if (ComposureProvider->FinalOutputRenderTarget)
+				{
+					MediaCapture->CaptureTextureRenderTarget2D(ComposureProvider->FinalOutputRenderTarget, Options);
+				}
+				else
+				{
+					UE_LOG(LogVCamOutputProvider, Warning, TEXT("MediaOutput mode - Composure usage was requested, but the specified ComposureOutputProvider has no FinalOutputRenderTarget set"));
+				}
+			}
+			else
+			{
+				TSharedPtr<FSceneViewport> SceneViewport = GetTargetSceneViewport();
+				if (SceneViewport.IsValid())
+				{
+					MediaCapture->CaptureSceneViewport(SceneViewport, Options);
+				}
+				else
+				{
+					UE_LOG(LogVCamOutputProvider, Warning, TEXT("MediaOutput mode - failed to find valid SceneViewport"));
+				}
+			}
+		}
+		else
+		{
+			UE_LOG(LogVCamOutputProvider, Warning, TEXT("MediaOutput mode - failed to create MediaCapture"));
+		}
+	}
+	else
+	{
+		UE_LOG(LogVCamOutputProvider, Warning, TEXT("MediaOutput mode - missing valid OutputConfig"));
+	}
+}
+
+void UVCamOutputMediaOutput::StopCapturing()
+{
+	if (MediaCapture)
+	{
+		MediaCapture->StopCapture(false);
+		MediaCapture->ConditionalBeginDestroy();
+		MediaCapture = nullptr;
+	}
+}
+
+#if WITH_EDITOR
+void UVCamOutputMediaOutput::PostEditChangeProperty(FPropertyChangedEvent& PropertyChangedEvent)
+{
+	FProperty* Property = PropertyChangedEvent.MemberProperty;
+
+	if (Property && PropertyChangedEvent.ChangeType != EPropertyChangeType::Interactive)
+	{
+		static FName NAME_OutputConfig = GET_MEMBER_NAME_CHECKED(UVCamOutputMediaOutput, OutputConfig);
+		static FName NAME_FromComposureOutputProviderIndex = GET_MEMBER_NAME_CHECKED(UVCamOutputMediaOutput, FromComposureOutputProviderIndex);
+
+		if ((Property->GetFName() == NAME_OutputConfig) ||
+			(Property->GetFName() == NAME_FromComposureOutputProviderIndex))
+		{
+			if (bIsActive)
+			{
+				SetActive(false);
+				SetActive(true);
+			}
+		}
+	}
+
+	Super::PostEditChangeProperty(PropertyChangedEvent);
+}
+#endif
diff --git a/Engine/Plugins/Experimental/VirtualCamera/Source/VCamCore/Private/Output/VCamOutputProviderBase.cpp b/Engine/Plugins/Experimental/VirtualCamera/Source/VCamCore/Private/Output/VCamOutputProviderBase.cpp
new file mode 100644
--- /dev/null
+++ b/Engine/Plugins/Experimental/VirtualCamera/Source/VCamCore/Private/Output/VCamOutputProviderBase.cpp
@@ -0,0 +1,332 @@
+﻿// Copyright Epic Games, Inc. All Rights Reserved.
+
+#include "Output/VCamOutputProviderBase.h"
+
+#include "Modifier/VCamModifierInterface.h"
+#include "VCamComponent.h"
+#include "UI/VCamWidget.h"
+
+#include "Engine/Engine.h"
+#include "UObject/UObjectBaseUtility.h"
+#include "Blueprint/UserWidget.h"
+#include "Blueprint/WidgetTree.h"
+
+#if WITH_EDITOR
+#include "Editor.h"
+#include "UnrealClient.h"
+#endif
+
+DEFINE_LOG_CATEGORY(LogVCamOutputProvider);
+
+void UVCamOutputProviderBase::BeginDestroy()
+{
+	Deinitialize();
+
+	Super::BeginDestroy();
+}
+
+void UVCamOutputProviderBase::Initialize()
+{
+	bool bWasInitialized = bInitialized;
+	bInitialized = true;
+
+	// Reactivate the provider if it was previously set to active
+	if (!bWasInitialized && bIsActive)
+	{
+#if WITH_EDITOR
+		// If the editor viewports aren't fully initialized, then delay initialization for the entire Output Provider
+		if (GEditor && GEditor->GetActiveViewport() && (GEditor->GetActiveViewport()->GetSizeXY().X < 1))
+		{
+			bInitialized = false;
+		}
+		else
+#endif
+		{
+			if (IsOuterComponentEnabled())
+			{
+				Activate();
+			}
+		}
+	}
+}
+
+void UVCamOutputProviderBase::Deinitialize()
+{
+	if (IsOuterComponentEnabled() && bIsActive)
+	{
+		Deactivate();
+	}
+	
+	bInitialized = false;
+}
+
+void UVCamOutputProviderBase::Activate()
+{
+	CreateUMG();
+	DisplayUMG();
+}
+
+void UVCamOutputProviderBase::Deactivate()
+{
+	DestroyUMG();
+}
+
+void UVCamOutputProviderBase::Tick(const float DeltaTime)
+{
+	if (bIsActive && UMGWidget && UMGClass)
+	{
+		UMGWidget->Tick(DeltaTime);
+	}
+}
+
+void UVCamOutputProviderBase::SetActive(const bool bInActive)
+{
+	bIsActive = bInActive;
+
+	if (IsOuterComponentEnabled())
+	{
+		if (bIsActive)
+		{
+			Activate();
+		}
+		else
+		{
+			Deactivate();
+		}
+	}
+}
+
+bool UVCamOutputProviderBase::IsOuterComponentEnabled() const
+{
+	if (UVCamComponent* OuterComponent = GetTypedOuter<UVCamComponent>())
+	{
+		return OuterComponent->IsEnabled();
+	}
+
+	return false;
+}
+
+
+void UVCamOutputProviderBase::SetTargetCamera(const UCineCameraComponent* InTargetCamera)
+{
+	TargetCamera = InTargetCamera;
+
+	NotifyWidgetOfComponentChange();
+}
+
+void UVCamOutputProviderBase::SetUMGClass(const TSubclassOf<UUserWidget> InUMGClass)
+{
+	UMGClass = InUMGClass;
+}
+
+void UVCamOutputProviderBase::CreateUMG()
+{
+	if (!UMGClass)
+	{
+		return;
+	}
+
+	if (UMGWidget)
+	{
+		UE_LOG(LogVCamOutputProvider, Error, TEXT("CreateUMG widget already set - failed to create"));
+		return;
+	}
+
+	UMGWidget = NewObject<UVPFullScreenUserWidget>(this, UVPFullScreenUserWidget::StaticClass());
+	UMGWidget->SetDisplayTypes(DisplayType, DisplayType, DisplayType);
+	UMGWidget->PostProcessDisplayType.bReceiveHardwareInput = true;
+
+#if WITH_EDITOR
+	UMGWidget->SetAllTargetViewports(GetTargetLevelViewport());
+#endif
+
+	UMGWidget->WidgetClass = UMGClass;
+	UE_LOG(LogVCamOutputProvider, Log, TEXT("CreateUMG widget named %s from class %s"), *UMGWidget->GetName(), *UMGWidget->WidgetClass->GetName());
+}
+
+void UVCamOutputProviderBase::DisplayUMG()
+{
+	if (UMGWidget)
+	{
+		UWorld* ActorWorld = nullptr;
+		int32 WorldType = -1;
+
+		for (const FWorldContext& Context : GEngine->GetWorldContexts())
+		{
+			if (Context.World())
+			{
+				// Prioritize PIE and Game modes if active
+				if ((Context.WorldType == EWorldType::PIE) || (Context.WorldType == EWorldType::Game))
+				{
+					ActorWorld = Context.World();
+					WorldType = (int32)Context.WorldType;
+					break;
+				}
+				else if (Context.WorldType == EWorldType::Editor)
+				{
+					// Only grab the Editor world if PIE and Game aren't available
+					ActorWorld = Context.World();
+					WorldType = (int32)Context.WorldType;
+				}
+			}
+		}
+
+		if (ActorWorld)
+		{
+			UMGWidget->Display(ActorWorld);
+			UE_LOG(LogVCamOutputProvider, Log, TEXT("DisplayUMG widget displayed in WorldType %d"), WorldType);
+		}
+
+		NotifyWidgetOfComponentChange();
+	}
+}
+
+void UVCamOutputProviderBase::DestroyUMG()
+{
+	if (UMGWidget)
+	{
+		if (UMGWidget->IsDisplayed())
+		{
+			UMGWidget->Hide();
+			UE_LOG(LogVCamOutputProvider, Log, TEXT("DestroyUMG widget %s hidden"), *UMGWidget->GetName());
+		}
+		UE_LOG(LogVCamOutputProvider, Log, TEXT("DestroyUMG widget %s destroyed"), *UMGWidget->GetName());
+
+#if WITH_EDITOR
+		UMGWidget->ResetAllTargetViewports();
+#endif
+
+		UMGWidget->ConditionalBeginDestroy();
+		UMGWidget = nullptr;
+	}
+}
+
+void UVCamOutputProviderBase::NotifyWidgetOfComponentChange() const
+{
+	if (UMGWidget && UMGWidget->IsDisplayed())
+	{
+		UUserWidget* DisplayedWidget = UMGWidget->GetWidget();
+		if (IsValid(DisplayedWidget))
+		{
+			if (UVCamComponent* OwningComponent = GetTypedOuter<UVCamComponent>())
+			{
+				UVCamComponent* VCamComponent = bIsActive ? OwningComponent : nullptr;
+
+				if (DisplayedWidget->Implements<UVCamModifierInterface>())
+				{
+					IVCamModifierInterface::Execute_OnVCamComponentChanged(DisplayedWidget, VCamComponent);
+				}
+
+				// Find all VCam Widgets inside the displayed widget and Initialize them with the owning VCam Component
+				if (IsValid(DisplayedWidget->WidgetTree))
+				{
+					DisplayedWidget->WidgetTree->ForEachWidget([VCamComponent](UWidget* Widget)
+					{
+						if (UVCamWidget* VCamWidget = Cast<UVCamWidget>(Widget))
+						{
+							VCamWidget->InitializeConnections(VCamComponent);
+						}
+					});
+				}
+			}
+		}
+
+	}
+}
+
+UVCamOutputProviderBase* UVCamOutputProviderBase::GetOtherOutputProviderByIndex(int32 Index) const
+{
+	if (Index > INDEX_NONE)
+	{
+		if (UVCamComponent* OuterComponent = GetTypedOuter<UVCamComponent>())
+		{
+			if (UVCamOutputProviderBase* Provider = OuterComponent->GetOutputProviderByIndex(Index))
+			{
+				return Provider;
+			}
+			else
+			{
+				UE_LOG(LogVCamOutputProvider, Warning, TEXT("GetOtherOutputProviderByIndex - specified index is out of range"));
+			}
+		}
+	}
+
+	return nullptr;
+}
+
+
+TSharedPtr<FSceneViewport> UVCamOutputProviderBase::GetTargetSceneViewport() const
+{
+	TSharedPtr<FSceneViewport> SceneViewport;
+
+	if (UVCamComponent* OuterComponent = GetTypedOuter<UVCamComponent>())
+	{
+		SceneViewport = OuterComponent->GetTargetSceneViewport();
+	}
+
+	return SceneViewport;
+}
+
+TWeakPtr<SWindow> UVCamOutputProviderBase::GetTargetInputWindow() const
+{
+	TWeakPtr<SWindow> InputWindow;
+
+	if (UVCamComponent* OuterComponent = GetTypedOuter<UVCamComponent>())
+	{
+		InputWindow = OuterComponent->GetTargetInputWindow();
+	}
+
+	return InputWindow;
+}
+
+#if WITH_EDITOR
+FLevelEditorViewportClient* UVCamOutputProviderBase::GetTargetLevelViewportClient() const
+{
+	FLevelEditorViewportClient* ViewportClient = nullptr;
+
+	if (UVCamComponent* OuterComponent = GetTypedOuter<UVCamComponent>())
+	{
+		ViewportClient = OuterComponent->GetTargetLevelViewportClient();
+	}
+
+	return ViewportClient;
+}
+
+TSharedPtr<SLevelViewport> UVCamOutputProviderBase::GetTargetLevelViewport() const
+{
+	TSharedPtr<SLevelViewport> LevelViewport;
+
+	if (UVCamComponent* OuterComponent = GetTypedOuter<UVCamComponent>())
+	{
+		LevelViewport = OuterComponent->GetTargetLevelViewport();
+	}
+
+	return LevelViewport;
+}
+
+void UVCamOutputProviderBase::PostEditChangeProperty(FPropertyChangedEvent& PropertyChangedEvent)
+{
+	FProperty* Property = PropertyChangedEvent.MemberProperty;
+
+	if (Property && PropertyChangedEvent.ChangeType != EPropertyChangeType::Interactive)
+	{
+		static FName NAME_IsActive = GET_MEMBER_NAME_CHECKED(UVCamOutputProviderBase, bIsActive);
+		static FName NAME_UMGClass = GET_MEMBER_NAME_CHECKED(UVCamOutputProviderBase, UMGClass);
+
+		if (Property->GetFName() == NAME_IsActive)
+		{
+			SetActive(bIsActive);
+		}
+		else if (Property->GetFName() == NAME_UMGClass)
+		{
+			if (bIsActive)
+			{
+				SetActive(false);
+				SetActive(true);
+			}
+		}
+	}
+
+	Super::PostEditChangeProperty(PropertyChangedEvent);
+}
+#endif
diff --git a/Engine/Plugins/Experimental/VirtualCamera/Source/VCamCore/Private/Output/VCamOutputRemoteSession.cpp b/Engine/Plugins/Experimental/VirtualCamera/Source/VCamCore/Private/Output/VCamOutputRemoteSession.cpp
new file mode 100644
--- /dev/null
+++ b/Engine/Plugins/Experimental/VirtualCamera/Source/VCamCore/Private/Output/VCamOutputRemoteSession.cpp
@@ -0,0 +1,399 @@
+// Copyright Epic Games, Inc. All Rights Reserved.
+
+#include "Output/VCamOutputRemoteSession.h"
+
+#include "Output/VCamOutputComposure.h"
+
+#include "Engine/GameInstance.h"
+#include "Engine/GameEngine.h"
+#include "Channels/RemoteSessionImageChannel.h"
+#include "Channels/RemoteSessionInputChannel.h"
+#include "GameFramework/PlayerController.h"
+#include "UObject/Package.h"
+#include "UObject/SoftObjectPath.h"
+#include "Slate/SceneViewport.h"
+#include "VPFullScreenUserWidget.h"
+#include "Widgets/SVirtualWindow.h"
+
+#if WITH_EDITOR
+#include "LevelEditor.h"
+#include "SLevelViewport.h"
+#endif
+
+namespace UE::VCamOutputRemoteSession::Private
+{
+	static const FSoftClassPath EmptyUMGSoftClassPath(TEXT("/VCamCore/Assets/VCam_EmptyVisibleUMG.VCam_EmptyVisibleUMG_C"));
+}
+
+void UVCamOutputRemoteSession::Initialize()
+{
+	if (!bInitialized && (MediaOutput == nullptr))
+	{
+		MediaOutput = NewObject<URemoteSessionMediaOutput>(GetTransientPackage(), URemoteSessionMediaOutput::StaticClass());
+	}
+
+	Super::Initialize();
+}
+
+void UVCamOutputRemoteSession::Deinitialize()
+{
+	MediaOutput = nullptr;
+
+	Super::Deinitialize();
+}
+
+void UVCamOutputRemoteSession::Activate()
+{
+	// If we don't have a UMG assigned, we still need to create an empty 'dummy' UMG in order to properly route the input back from the RemoteSession device
+	if (UMGClass == nullptr)
+	{
+		bUsingDummyUMG = true;
+		UMGClass = UE::VCamOutputRemoteSession::Private::EmptyUMGSoftClassPath.TryLoadClass<UUserWidget>();
+	}
+
+	CreateRemoteSession();
+	
+	Super::Activate();
+}
+
+void UVCamOutputRemoteSession::Deactivate()
+{
+	DestroyRemoteSession();
+
+	Super::Deactivate();
+
+	if (bUsingDummyUMG)
+	{
+		UMGClass = nullptr;
+		bUsingDummyUMG = false;
+	}
+}
+
+void UVCamOutputRemoteSession::Tick(const float DeltaTime)
+{
+	if (bIsActive && RemoteSessionHost.IsValid())
+	{
+		RemoteSessionHost->Tick(DeltaTime);
+	}
+
+	Super::Tick(DeltaTime);
+}
+
+void UVCamOutputRemoteSession::CreateRemoteSession()
+{
+	if (!bInitialized)
+	{
+		UE_LOG(LogVCamOutputProvider, Warning, TEXT("CreateRemoteSession has been called, but has not been initialized yet"));
+		return;
+	}
+
+	if (!RemoteSessionHost.IsValid())
+	{
+		if (IRemoteSessionModule* RemoteSession = FModuleManager::LoadModulePtr<IRemoteSessionModule>("RemoteSession"))
+		{
+			TArray<FRemoteSessionChannelInfo> SupportedChannels;
+			SupportedChannels.Emplace(FRemoteSessionInputChannel::StaticType(), ERemoteSessionChannelMode::Read);
+			SupportedChannels.Emplace(FRemoteSessionImageChannel::StaticType(), ERemoteSessionChannelMode::Write);
+
+			RemoteSessionHost = RemoteSession->CreateHost(MoveTemp(SupportedChannels), PortNumber);
+
+			if (RemoteSessionHost.IsValid())
+			{
+				RemoteSessionHost->RegisterChannelChangeDelegate(FOnRemoteSessionChannelChange::FDelegate::CreateUObject(this, &UVCamOutputRemoteSession::OnRemoteSessionChannelChange));
+
+				RemoteSessionHost->Tick(0.0f);
+
+				if (bUseOverrideResolution)
+				{
+					TSharedPtr<FSceneViewport> SceneViewport = GetTargetSceneViewport();
+					if (SceneViewport.IsValid())
+					{
+						SceneViewport->SetFixedViewportSize(OverrideResolution.X, OverrideResolution.Y);
+					}
+				}
+			}
+			else
+			{
+				UE_LOG(LogVCamOutputProvider, Error, TEXT("Unable to create a RemoteSession host on port %d. Check the port is not in use."), PortNumber);
+			}			
+		}
+	}
+	else
+	{
+		UE_LOG(LogVCamOutputProvider, Warning, TEXT("CreateRemoteSession has been called, but there is already a valid RemoteSessionHost!"));
+	}
+}
+
+void UVCamOutputRemoteSession::DestroyRemoteSession()
+{
+	if (RemoteSessionHost.IsValid())
+	{
+		TSharedPtr<FRemoteSessionInputChannel> InputChannel = RemoteSessionHost->GetChannel<FRemoteSessionInputChannel>();
+		if (InputChannel)
+		{
+			InputChannel->GetOnRouteTouchDownToWidgetFailedDelegate()->RemoveAll(this);
+		}
+
+		RemoteSessionHost->Close(TEXT("Camera Shutdown"));
+		RemoteSessionHost.Reset();
+		if (MediaCapture)
+		{
+			MediaCapture->StopCapture(true);
+			MediaCapture = nullptr;
+		}
+
+		// If everything is being shut down somehow, get out now
+		if (GEngine == nullptr)
+		{
+			return;
+		}
+
+#if WITH_EDITOR
+		for (const FWorldContext& Context : GEngine->GetWorldContexts())
+		{
+			if (Context.WorldType == EWorldType::Editor)
+			{
+				if (FLevelEditorViewportClient* LevelViewportClient = GetTargetLevelViewportClient())
+				{
+					LevelViewportClient->ViewFOV = LevelViewportClient->FOVAngle;
+					GEditor->RemovePerspectiveViewRotation(true, true, false);
+				}
+			}
+		}
+#endif
+
+		if (bUseOverrideResolution)
+		{
+			TSharedPtr<FSceneViewport> SceneViewport = GetTargetSceneViewport();
+			if (SceneViewport.IsValid())
+			{
+				SceneViewport->SetFixedViewportSize(0, 0);
+			}
+		}
+	}
+}
+
+void UVCamOutputRemoteSession::OnRemoteSessionChannelChange(IRemoteSessionRole* Role, TWeakPtr<IRemoteSessionChannel> Channel, ERemoteSessionChannelChange Change)
+{
+	TSharedPtr<IRemoteSessionChannel> PinnedChannel = Channel.Pin();
+	if (PinnedChannel && Change == ERemoteSessionChannelChange::Created)
+	{
+        const FString PinnedChannelTypeString(PinnedChannel->GetType());
+        
+		if (PinnedChannelTypeString.Compare(FRemoteSessionImageChannel::StaticType()) == 0)
+		{
+			OnImageChannelCreated(Channel);
+		}
+        else if (PinnedChannelTypeString.Compare(FRemoteSessionInputChannel::StaticType()) == 0)
+		{
+			OnInputChannelCreated(Channel);
+		}
+	}
+}
+
+void UVCamOutputRemoteSession::OnImageChannelCreated(TWeakPtr<IRemoteSessionChannel> Instance)
+{
+	TSharedPtr<FRemoteSessionImageChannel> ImageChannel = StaticCastSharedPtr<FRemoteSessionImageChannel>(Instance.Pin());
+	if (ImageChannel)
+	{
+		ImageChannel->SetImageProvider(nullptr);
+		MediaOutput->SetImageChannel(ImageChannel);
+		MediaCapture = Cast<URemoteSessionMediaCapture>(MediaOutput->CreateMediaCapture());
+
+		FMediaCaptureOptions Options;
+		Options.bResizeSourceBuffer = true;
+
+		// If we are rendering from a ComposureOutputProvider, get the requested render target and use that instead of the viewport
+		if (UVCamOutputComposure* ComposureProvider = Cast<UVCamOutputComposure>(GetOtherOutputProviderByIndex(FromComposureOutputProviderIndex)))
+		{
+			if (ComposureProvider->FinalOutputRenderTarget)
+			{
+				MediaCapture->CaptureTextureRenderTarget2D(ComposureProvider->FinalOutputRenderTarget, Options);
+				UE_LOG(LogVCamOutputProvider, Log, TEXT("ImageChannel callback - MediaCapture set with ComposureRenderTarget"));
+			}
+			else
+			{
+				UE_LOG(LogVCamOutputProvider, Warning, TEXT("ImageChannel callback - Composure usage was requested, but the specified ComposureOutputProvider has no FinalOutputRenderTarget set"));
+			}
+		}
+		else
+		{
+			TWeakPtr<FSceneViewport> SceneViewport = GetTargetSceneViewport();
+			if (TSharedPtr<FSceneViewport> PinnedSceneViewport = SceneViewport.Pin())
+			{
+				MediaCapture->CaptureSceneViewport(PinnedSceneViewport, Options);
+				UE_LOG(LogVCamOutputProvider, Log, TEXT("ImageChannel callback - MediaCapture set with viewport"));
+			}
+		}
+	}
+}
+
+void UVCamOutputRemoteSession::OnInputChannelCreated(TWeakPtr<IRemoteSessionChannel> Instance)
+{
+	TSharedPtr<FRemoteSessionInputChannel> InputChannel = StaticCastSharedPtr<FRemoteSessionInputChannel>(Instance.Pin());
+	if (InputChannel)
+	{
+		// If we have a UMG, then use it
+		if (UMGClass && UMGWidget)
+		{
+			TSharedPtr<SVirtualWindow> InputWindow;
+
+			// If we are rendering from a ComposureOutputProvider, we need to get the InputWindow from that UMG, not the one in the RemoteSessionOutputProvider
+			if (UVCamOutputComposure* ComposureProvider = Cast<UVCamOutputComposure>(GetOtherOutputProviderByIndex(FromComposureOutputProviderIndex)))
+			{
+				if (UVPFullScreenUserWidget* ComposureUMGWidget = ComposureProvider->GetUMGWidget())
+				{
+					InputWindow = ComposureUMGWidget->PostProcessDisplayType.GetSlateWindow();
+					UE_LOG(LogVCamOutputProvider, Log, TEXT("InputChannel callback - Routing input to active viewport with Composure UMG"));
+				}
+				else
+				{
+					UE_LOG(LogVCamOutputProvider, Warning, TEXT("InputChannel callback - Composure usage was requested, but the specified ComposureOutputProvider has no UMG set"));
+				}
+			}
+			else
+			{
+				InputWindow = UMGWidget->PostProcessDisplayType.GetSlateWindow();
+				UE_LOG(LogVCamOutputProvider, Log, TEXT("InputChannel callback - Routing input to active viewport with UMG"));
+			}
+
+			InputChannel->SetPlaybackWindow(InputWindow, nullptr);
+			InputChannel->TryRouteTouchMessageToWidget(true);
+			InputChannel->GetOnRouteTouchDownToWidgetFailedDelegate()->AddUObject(this, &UVCamOutputRemoteSession::OnTouchEventOutsideUMG);
+		}
+		else
+		{
+			InputChannel->SetPlaybackWindow(GetTargetInputWindow(), nullptr);
+			InputChannel->TryRouteTouchMessageToWidget(true);
+			InputChannel->GetOnRouteTouchDownToWidgetFailedDelegate()->AddUObject(this, &UVCamOutputRemoteSession::OnTouchEventOutsideUMG);
+
+			UE_LOG(LogVCamOutputProvider, Log, TEXT("InputChannel callback - Routing input to active viewport"));
+		}
+	}
+}
+
+void UVCamOutputRemoteSession::OnTouchEventOutsideUMG(const FVector2D& InViewportPosition)
+{
+	const float MaxFocusTraceDistance = 1000000.0f;
+	FVector TraceDirection;
+	FVector CameraWorldLocation;
+
+	if (!DeprojectScreenToWorld(InViewportPosition, CameraWorldLocation, TraceDirection))
+	{
+		return;
+	}
+
+	for (const FWorldContext& Context : GEngine->GetWorldContexts())
+	{
+		UWorld* World = Context.World();
+		if (World)
+		{
+			FCollisionQueryParams TraceParams(SCENE_QUERY_STAT(UpdateAutoFocus), true);
+			const FVector TraceEnd = CameraWorldLocation + TraceDirection * MaxFocusTraceDistance;
+			const bool bHit = World->LineTraceSingleByChannel(LastViewportTouchResult, CameraWorldLocation, TraceEnd, ECC_Visibility, TraceParams);
+			if (bHit)
+			{
+#if WITH_EDITOR
+				// @todo: This doesn't seem like the most efficient way to pass click events to the editor viewport...
+				if (Context.WorldType == EWorldType::Editor)
+				{
+					if (FLevelEditorViewportClient* LevelViewportClient = GetTargetLevelViewportClient())
+					{
+						if (FViewport* ActiveViewport = LevelViewportClient->Viewport)
+						{
+							FSceneViewFamilyContext ViewFamily(FSceneViewFamily::ConstructionValues(
+								ActiveViewport,
+								LevelViewportClient->GetScene(),
+								LevelViewportClient->EngineShowFlags)
+								.SetRealtimeUpdate(true));
+							FSceneView* View = LevelViewportClient->CalcSceneView(&ViewFamily);
+
+							uint32 HitX = InViewportPosition.X;
+							uint32 HitY = InViewportPosition.Y;
+							HHitProxy* HitProxy = ActiveViewport->GetHitProxy(HitX, HitY);
+							LevelViewportClient->ProcessClick(*View, HitProxy, EKeys::LeftMouseButton, EInputEvent::IE_Pressed, HitX, HitY);
+						}
+					}
+				}
+				else
+#endif
+				{
+					LastViewportTouchResult.GetHitObjectHandle().FetchActor()->NotifyActorOnClicked();
+				}
+			}
+		}
+	}
+}
+
+bool UVCamOutputRemoteSession::DeprojectScreenToWorld(const FVector2D& InScreenPosition, FVector& OutWorldPosition, FVector& OutWorldDirection) const
+{
+	bool bSuccess = false;
+
+	for (const FWorldContext& Context : GEngine->GetWorldContexts())
+	{
+		if (Context.WorldType == EWorldType::PIE || Context.WorldType == EWorldType::Game)
+		{
+			APlayerController* PC = Context.OwningGameInstance->GetFirstLocalPlayerController(Context.World());
+			if (PC)
+			{
+				bSuccess |= PC->DeprojectScreenPositionToWorld(InScreenPosition.X, InScreenPosition.Y, OutWorldPosition, OutWorldDirection);
+				break;
+			}
+		}
+#if WITH_EDITOR
+		else if (Context.WorldType == EWorldType::Editor)
+		{
+			if (FLevelEditorViewportClient* LevelViewportClient = GetTargetLevelViewportClient())
+			{
+				if (FViewport* ActiveViewport = LevelViewportClient->Viewport)
+				{
+					FSceneViewFamilyContext ViewFamily(FSceneViewFamily::ConstructionValues(
+						ActiveViewport,
+						LevelViewportClient->GetScene(),
+						LevelViewportClient->EngineShowFlags)
+						.SetRealtimeUpdate(true));
+					FSceneView* View = LevelViewportClient->CalcSceneView(&ViewFamily);
+
+					const FIntPoint ViewportSize = ActiveViewport->GetSizeXY();
+					const FIntRect ViewRect = FIntRect(0, 0, ViewportSize.X, ViewportSize.Y);
+					const FMatrix InvViewProjectionMatrix = View->ViewMatrices.GetInvViewProjectionMatrix();
+					FSceneView::DeprojectScreenToWorld(InScreenPosition, ViewRect, InvViewProjectionMatrix, OutWorldPosition, OutWorldDirection);
+					bSuccess = true;
+				}
+			}
+		}
+#endif
+	}
+
+	if (!bSuccess)
+	{
+		OutWorldPosition = FVector::ZeroVector;
+		OutWorldDirection = FVector::ZeroVector;
+	}
+	return bSuccess;
+}
+
+#if WITH_EDITOR
+void UVCamOutputRemoteSession::PostEditChangeProperty(FPropertyChangedEvent& PropertyChangedEvent)
+{
+	FProperty* Property = PropertyChangedEvent.MemberProperty;
+
+	if (Property && PropertyChangedEvent.ChangeType != EPropertyChangeType::Interactive)
+	{
+		static FName NAME_PortNumber = GET_MEMBER_NAME_CHECKED(UVCamOutputRemoteSession, PortNumber);
+		static FName NAME_FromComposureOutputProviderIndex = GET_MEMBER_NAME_CHECKED(UVCamOutputRemoteSession, FromComposureOutputProviderIndex);
+
+		if ((Property->GetFName() == NAME_PortNumber) ||
+			(Property->GetFName() == NAME_FromComposureOutputProviderIndex))
+		{
+			if (bIsActive)
+			{
+				SetActive(false);
+				SetActive(true);
+			}
+		}
+	}
+
+	Super::PostEditChangeProperty(PropertyChangedEvent);
+}
+#endif
diff --git a/Engine/Plugins/Experimental/VirtualCamera/Source/VCamCore/Private/Output/VCamOutputViewport.cpp b/Engine/Plugins/Experimental/VirtualCamera/Source/VCamCore/Private/Output/VCamOutputViewport.cpp
new file mode 100644
--- /dev/null
+++ b/Engine/Plugins/Experimental/VirtualCamera/Source/VCamCore/Private/Output/VCamOutputViewport.cpp
@@ -0,0 +1,10 @@
+// Copyright Epic Games, Inc. All Rights Reserved.
+
+#include "Output/VCamOutputViewport.h"
+
+void UVCamOutputViewport::CreateUMG()
+{
+	DisplayType = EVPWidgetDisplayType::Viewport;
+
+	Super::CreateUMG();
+}
diff --git a/Engine/Plugins/Experimental/VirtualCamera/Source/VCamCore/Private/VCamComponent.cpp b/Engine/Plugins/Experimental/VirtualCamera/Source/VCamCore/Private/VCamComponent.cpp
--- a/Engine/Plugins/Experimental/VirtualCamera/Source/VCamCore/Private/VCamComponent.cpp
+++ b/Engine/Plugins/Experimental/VirtualCamera/Source/VCamCore/Private/VCamComponent.cpp
@@ -2,24 +2,23 @@
 
 #include "VCamComponent.h"
 
+#include "Modifier/VCamModifier.h"
+#include "Modifier/VCamModifierContext.h"
+
 #include "CineCameraComponent.h"
 #include "Engine/GameInstance.h"
-#include "Features/IModularFeatures.h"
-#include "Framework/Application/SlateApplication.h"
-#include "ILiveLinkClient.h"
-#include "VCamModifier.h"
-#include "VCamModifierContext.h"
-#include "Roles/LiveLinkCameraRole.h"
-#include "Roles/LiveLinkTransformRole.h"
-#include "GameDelegates.h"
 #include "Engine/GameEngine.h"
+#include "Engine/InputDelegateBinding.h"
 #include "EnhancedInputComponent.h"
 #include "EnhancedInputSubsystems.h"
-#include "Engine/InputDelegateBinding.h"
+#include "Features/IModularFeatures.h"
+#include "Framework/Application/SlateApplication.h"
+#include "GameFramework/InputSettings.h"
+#include "ILiveLinkClient.h"
 #include "InputMappingContext.h"
 #include "PlayerMappableKeySettings.h"
-
-#include "GameFramework/InputSettings.h"
+#include "Roles/LiveLinkCameraRole.h"
+#include "Roles/LiveLinkTransformRole.h"
 
 #if WITH_EDITOR
 #include "Modules/ModuleManager.h"
diff --git a/Engine/Plugins/Experimental/VirtualCamera/Source/VCamCore/Private/VCamModifier.cpp b/Engine/Plugins/Experimental/VirtualCamera/Source/VCamCore/Private/VCamModifier.cpp
deleted file mode 100644
diff --git a/Engine/Plugins/Experimental/VirtualCamera/Source/VCamCore/Private/VCamModifierInterface.cpp b/Engine/Plugins/Experimental/VirtualCamera/Source/VCamCore/Private/VCamModifierInterface.cpp
deleted file mode 100644
diff --git a/Engine/Plugins/Experimental/VirtualCamera/Source/VCamCore/Private/VCamOutputComposure.cpp b/Engine/Plugins/Experimental/VirtualCamera/Source/VCamCore/Private/VCamOutputComposure.cpp
deleted file mode 100644
diff --git a/Engine/Plugins/Experimental/VirtualCamera/Source/VCamCore/Private/VCamOutputMediaOutput.cpp b/Engine/Plugins/Experimental/VirtualCamera/Source/VCamCore/Private/VCamOutputMediaOutput.cpp
deleted file mode 100644
diff --git a/Engine/Plugins/Experimental/VirtualCamera/Source/VCamCore/Private/VCamOutputProviderBase.cpp b/Engine/Plugins/Experimental/VirtualCamera/Source/VCamCore/Private/VCamOutputProviderBase.cpp
deleted file mode 100644
diff --git a/Engine/Plugins/Experimental/VirtualCamera/Source/VCamCore/Private/VCamOutputRemoteSession.cpp b/Engine/Plugins/Experimental/VirtualCamera/Source/VCamCore/Private/VCamOutputRemoteSession.cpp
deleted file mode 100644
diff --git a/Engine/Plugins/Experimental/VirtualCamera/Source/VCamCore/Private/VCamOutputViewport.cpp b/Engine/Plugins/Experimental/VirtualCamera/Source/VCamCore/Private/VCamOutputViewport.cpp
deleted file mode 100644
diff --git a/Engine/Plugins/Experimental/VirtualCamera/Source/VCamCore/Public/EVCamTargetViewportID.h b/Engine/Plugins/Experimental/VirtualCamera/Source/VCamCore/Public/EVCamTargetViewportID.h
new file mode 100644
--- /dev/null
+++ b/Engine/Plugins/Experimental/VirtualCamera/Source/VCamCore/Public/EVCamTargetViewportID.h
@@ -0,0 +1,16 @@
+// Copyright Epic Games, Inc. All Rights Reserved.
+
+#pragma once
+
+#include "CoreMinimal.h"
+#include "EVCamTargetViewportID.generated.h"
+
+UENUM(BlueprintType, meta=(DisplayName = "VCam Target Viewport ID"))
+enum class EVCamTargetViewportID : uint8
+{
+	CurrentlySelected = 0,
+	Viewport1 = 1,
+	Viewport2 = 2,
+	Viewport3 = 3,
+	Viewport4 = 4
+};
\ No newline at end of file
diff --git a/Engine/Plugins/Experimental/VirtualCamera/Source/VCamCore/Public/Modifier/ModifierStackEntry.h b/Engine/Plugins/Experimental/VirtualCamera/Source/VCamCore/Public/Modifier/ModifierStackEntry.h
new file mode 100644
--- /dev/null
+++ b/Engine/Plugins/Experimental/VirtualCamera/Source/VCamCore/Public/Modifier/ModifierStackEntry.h
@@ -0,0 +1,79 @@
+// Copyright Epic Games, Inc. All Rights Reserved.
+
+#pragma once
+
+#include "CoreMinimal.h"
+#include "VCamModifier.h"
+#include "ModifierStackEntry.generated.h"
+
+class UVCamModifier;
+
+// Links a Modifier with a Name for use in a Modifier Stack
+USTRUCT()
+struct FModifierStackEntry
+{
+	GENERATED_BODY()
+
+	// Identifier for this modifier in the stack
+    UPROPERTY(EditAnywhere, Category="Modifier")
+    FName Name;
+
+	// Controls whether the modifier actually gets applied
+	UPROPERTY(EditAnywhere, Category="Modifier")
+	bool bEnabled = true;
+
+	// The current generated modifier instance
+	UPROPERTY(EditAnywhere, Instanced, Category="Modifier")
+    TObjectPtr<UVCamModifier> GeneratedModifier = nullptr;
+
+#if WITH_EDITOR
+	// GUID used in the editor to identify specific stack entries during editor operations
+	FGuid EditorGuid;
+#endif
+	
+	FModifierStackEntry()
+	{
+#if WITH_EDITOR
+		EditorGuid = FGuid::NewGuid();
+#endif
+	}
+
+	// If ModifierClass is provided then you must also supply a valid outer for the generated modifier
+	FModifierStackEntry(const FName& InName, const TSubclassOf<UVCamModifier> InModifierClass, UObject* InOuter)
+		: Name(InName)
+	{
+		if (InModifierClass)
+		{
+			GeneratedModifier = NewObject<UVCamModifier>(InOuter, InModifierClass.Get());
+		}
+
+#if WITH_EDITOR
+		EditorGuid = FGuid::NewGuid();
+#endif
+	};
+
+	~FModifierStackEntry()
+	{
+		// It's possible the Generated Modifier may pass IsValid while having already been destroyed so ensure we explicitly check for this case
+		if (IsValid(GeneratedModifier) && !GeneratedModifier->HasAnyFlags(RF_BeginDestroyed | RF_FinishDestroyed))
+		{
+			GeneratedModifier->Deinitialize();
+		}
+	}
+
+	bool operator==(const FModifierStackEntry& Other) const
+	{
+		bool bArePropertiesEqual = this->Name.IsEqual(Other.Name) && this->bEnabled == Other.bEnabled && this->GeneratedModifier == Other.GeneratedModifier;
+
+#if WITH_EDITOR
+		bArePropertiesEqual = bArePropertiesEqual && this->EditorGuid == Other.EditorGuid;
+#endif
+
+		return bArePropertiesEqual;
+	};
+
+	bool operator!=(const FModifierStackEntry& Other) const
+	{
+		return !(*this == Other);
+	}
+};
diff --git a/Engine/Plugins/Experimental/VirtualCamera/Source/VCamCore/Public/Modifier/VCamModifier.h b/Engine/Plugins/Experimental/VirtualCamera/Source/VCamCore/Public/Modifier/VCamModifier.h
new file mode 100644
--- /dev/null
+++ b/Engine/Plugins/Experimental/VirtualCamera/Source/VCamCore/Public/Modifier/VCamModifier.h
@@ -0,0 +1,107 @@
+// Copyright Epic Games, Inc. All Rights Reserved.
+
+#pragma once
+
+#include "CoreMinimal.h"
+#include "CineCameraComponent.h"
+#include "Roles/LiveLinkCameraTypes.h"
+#include "VCamModifier.generated.h"
+
+class UInputAction;
+class UInputMappingContext;
+class UVCamComponent;
+class UVCamModifierContext;
+class UInputComponent;
+
+struct FModifierStackEntry;
+
+USTRUCT(BlueprintType)
+struct VCAMCORE_API FVCamModifierConnectionPoint
+{
+	GENERATED_BODY()
+
+	// An optional action to associate with this connection point
+	UPROPERTY(EditAnywhere, BlueprintReadWrite, Category = "Connection")
+	TObjectPtr<UInputAction> AssociatedAction;
+};
+
+UCLASS(Blueprintable, Abstract, EditInlineNew)
+class VCAMCORE_API UVCamModifier : public UObject
+{
+	GENERATED_BODY()
+
+public:
+	virtual void Initialize(UVCamModifierContext* Context, UInputComponent* InputComponent = nullptr);
+	
+	virtual void Deinitialize();
+
+	virtual void Apply(UVCamModifierContext* Context, UCineCameraComponent* CameraComponent, const float DeltaTime) {};
+
+	virtual void PostLoad() override;
+
+	bool DoesRequireInitialization() const { return bRequiresInitialization; };
+
+	UFUNCTION(BlueprintCallable, Category="VirtualCamera")
+	UVCamComponent* GetOwningVCamComponent() const;
+
+	UFUNCTION(BlueprintCallable, Category = "VirtualCamera")
+	void GetCurrentLiveLinkDataFromOwningComponent(FLiveLinkCameraBlueprintData& LiveLinkData);
+
+	UFUNCTION(BlueprintCallable, Category = "VirtualCamera")
+	void SetEnabled(bool bNewEnabled);
+
+	UFUNCTION(BlueprintPure, Category = "VirtualCamera", meta=(ReturnDisplayName="Enabled"))
+	bool IsEnabled() const;
+
+	// Sets the name of the modifier in the associated modifier stack
+	// Returns a bool for whether 
+	UFUNCTION(BlueprintCallable, Category = "VirtualCamera", meta=(ReturnDisplayName="Success"))
+	bool SetStackEntryName(FName NewName);
+
+	// Gets the name of the modifier in the associated modifier stack
+	UFUNCTION(BlueprintPure, Category = "VirtualCamera", meta=(ReturnDisplayName="Name"))
+	FName GetStackEntryName() const;
+
+	// If an Input Mapping Context is specified then that Context will be automatically added to the input system when this Modifier is Initialized 
+	UPROPERTY(EditAnywhere, BlueprintReadOnly, Category="VCam Input", meta=(EditCondition="bRegisterForInput"))
+	TObjectPtr<UInputMappingContext> InputMappingContext = nullptr;
+
+	// If an Input Mapping Context is provided then this value defines the priority level that the context is added to the input system with
+	UPROPERTY(EditAnywhere, BlueprintReadOnly, Category="VCam Input", meta=(EditCondition="bRegisterForInput"))
+	int32 InputContextPriority = 0;
+
+	UPROPERTY(EditAnywhere, BlueprintReadWrite, Category = "VCam Connection Points")
+	TMap<FName, FVCamModifierConnectionPoint> ConnectionPoints;
+
+	virtual UWorld* GetWorld() const override;
+	
+private:
+	FModifierStackEntry* GetCorrespondingStackEntry() const;
+
+	bool bRequiresInitialization = true;
+};
+
+UCLASS(EditInlineNew)
+class VCAMCORE_API UVCamBlueprintModifier : public UVCamModifier
+{
+	GENERATED_BODY()
+
+public:
+	virtual void Initialize(UVCamModifierContext* Context, UInputComponent* InputComponent=nullptr) override;
+	virtual void Deinitialize() override;
+	virtual void Apply(UVCamModifierContext* Context, UCineCameraComponent* CameraComponent, const float DeltaTime) override;
+	
+	UFUNCTION(BlueprintImplementableEvent, Category="VirtualCamera")
+	void OnInitialize(UVCamModifierContext* Context);
+
+	UFUNCTION(BlueprintImplementableEvent, Category="VirtualCamera")
+	void OnDeinitialize();
+
+	UFUNCTION(BlueprintImplementableEvent, Category="VirtualCamera")
+	void OnApply(UVCamModifierContext* Context, UCineCameraComponent* CameraComponent, const float DeltaTime);
+
+	// This function is deliberately non-working to force cleanup of Input Contexts
+	// Please move any previous values to the new properties in Class Defaults
+	UFUNCTION(BlueprintImplementableEvent, Category="VirtualCamera")
+	void GetInputMappingContextAndPriority() const;
+};
\ No newline at end of file
diff --git a/Engine/Plugins/Experimental/VirtualCamera/Source/VCamCore/Public/Modifier/VCamModifierContext.h b/Engine/Plugins/Experimental/VirtualCamera/Source/VCamCore/Public/Modifier/VCamModifierContext.h
new file mode 100644
--- /dev/null
+++ b/Engine/Plugins/Experimental/VirtualCamera/Source/VCamCore/Public/Modifier/VCamModifierContext.h
@@ -0,0 +1,12 @@
+// Copyright Epic Games, Inc. All Rights Reserved.
+
+#pragma once
+
+#include "CoreMinimal.h"
+#include "VCamModifierContext.generated.h"
+
+UCLASS(Blueprintable, Abstract, EditInlineNew)
+class VCAMCORE_API UVCamModifierContext : public UObject
+{
+	GENERATED_BODY()
+};
\ No newline at end of file
diff --git a/Engine/Plugins/Experimental/VirtualCamera/Source/VCamCore/Public/Modifier/VCamModifierInterface.h b/Engine/Plugins/Experimental/VirtualCamera/Source/VCamCore/Public/Modifier/VCamModifierInterface.h
new file mode 100644
--- /dev/null
+++ b/Engine/Plugins/Experimental/VirtualCamera/Source/VCamCore/Public/Modifier/VCamModifierInterface.h
@@ -0,0 +1,29 @@
+// Copyright Epic Games, Inc. All Rights Reserved.
+
+#pragma once
+
+#include "CoreMinimal.h"
+#include "UObject/Interface.h"
+#include "VCamModifierInterface.generated.h"
+
+class UVCamComponent;
+
+UINTERFACE(MinimalAPI,BlueprintType,Blueprintable)
+class UVCamModifierInterface : public UInterface
+{
+	GENERATED_BODY()
+};
+
+/**
+   The interface for all VCam modifier Blueprints to inherit so that we can enforce some common behaviors.
+ */
+class VCAMCORE_API IVCamModifierInterface
+{
+	GENERATED_BODY()
+public:
+	//Function used to monitor the status of a VCam. If it changed, function will be triggered.
+	UFUNCTION(BlueprintCallable, BlueprintNativeEvent, Category = "Active VCam Update")
+	void OnVCamComponentChanged(UVCamComponent* VCam);
+
+	virtual void OnVCamComponentChanged_Implementation(UVCamComponent* VCam);
+};
diff --git a/Engine/Plugins/Experimental/VirtualCamera/Source/VCamCore/Public/Output/VCamOutputComposure.h b/Engine/Plugins/Experimental/VirtualCamera/Source/VCamCore/Public/Output/VCamOutputComposure.h
new file mode 100644
--- /dev/null
+++ b/Engine/Plugins/Experimental/VirtualCamera/Source/VCamCore/Public/Output/VCamOutputComposure.h
@@ -0,0 +1,36 @@
+// Copyright Epic Games, Inc. All Rights Reserved.
+
+#pragma once
+
+#include "CoreMinimal.h"
+#include "VCamOutputProviderBase.h"
+#include "Engine/TextureRenderTarget2D.h"
+#include "CompositingElement.h"
+
+#if WITH_EDITOR
+#endif
+
+#include "VCamOutputComposure.generated.h"
+
+UCLASS(meta = (DisplayName = "Composure Output Provider"))
+class VCAMCORE_API UVCamOutputComposure : public UVCamOutputProviderBase
+{
+	GENERATED_BODY()
+
+public:
+
+#if WITH_EDITOR
+	virtual void PostEditChangeProperty(FPropertyChangedEvent& PropertyChangedEvent) override;
+#endif
+
+	// List of Composure stack Compositing Elements to render the requested UMG into
+	UPROPERTY(EditAnywhere, BlueprintReadWrite, Category = "Output")
+	TArray<TSoftObjectPtr<ACompositingElement>> LayerTargets;
+
+	// TextureRenderTarget2D asset that contains the final output
+	UPROPERTY(EditAnywhere, BlueprintReadWrite, Category = "Output")
+	TObjectPtr<UTextureRenderTarget2D> FinalOutputRenderTarget = nullptr;
+
+protected:
+	virtual void CreateUMG() override;
+};
diff --git a/Engine/Plugins/Experimental/VirtualCamera/Source/VCamCore/Public/Output/VCamOutputMediaOutput.h b/Engine/Plugins/Experimental/VirtualCamera/Source/VCamCore/Public/Output/VCamOutputMediaOutput.h
new file mode 100644
--- /dev/null
+++ b/Engine/Plugins/Experimental/VirtualCamera/Source/VCamCore/Public/Output/VCamOutputMediaOutput.h
@@ -0,0 +1,42 @@
+// Copyright Epic Games, Inc. All Rights Reserved.
+
+#pragma once
+
+#include "VCamOutputProviderBase.h"
+#include "MediaCapture.h"
+#include "MediaOutput.h"
+
+#if WITH_EDITOR
+#endif
+
+#include "VCamOutputMediaOutput.generated.h"
+
+UCLASS(meta = (DisplayName = "Media Output Provider"))
+class VCAMCORE_API UVCamOutputMediaOutput: public UVCamOutputProviderBase
+{
+	GENERATED_BODY()
+
+public:
+	virtual void Activate() override;
+	virtual void Deactivate() override;
+
+#if WITH_EDITOR
+	virtual void PostEditChangeProperty(FPropertyChangedEvent& PropertyChangedEvent) override;
+#endif
+
+	// Media Output configuration asset
+	UPROPERTY(EditAnywhere, BlueprintReadWrite, Category = "Output")
+	TObjectPtr<UMediaOutput> OutputConfig;
+
+	// If using the output from a Composure Output Provider, specify it here
+	UPROPERTY(EditAnywhere, BlueprintReadWrite, Category = "Output")
+	int32 FromComposureOutputProviderIndex = INDEX_NONE;
+
+protected:
+	UPROPERTY(Transient)
+	TObjectPtr<UMediaCapture> MediaCapture = nullptr;
+
+private:
+	void StartCapturing();
+	void StopCapturing();
+};
diff --git a/Engine/Plugins/Experimental/VirtualCamera/Source/VCamCore/Public/Output/VCamOutputProviderBase.h b/Engine/Plugins/Experimental/VirtualCamera/Source/VCamCore/Public/Output/VCamOutputProviderBase.h
new file mode 100644
--- /dev/null
+++ b/Engine/Plugins/Experimental/VirtualCamera/Source/VCamCore/Public/Output/VCamOutputProviderBase.h
@@ -0,0 +1,140 @@
+﻿// Copyright Epic Games, Inc. All Rights Reserved.
+
+#pragma once
+
+#include "CoreMinimal.h"
+#include "CineCameraComponent.h"
+#include "VPFullScreenUserWidget.h"
+#include "VCamOutputProviderBase.generated.h"
+
+DECLARE_LOG_CATEGORY_EXTERN(LogVCamOutputProvider, Log, All);
+
+class UUserWidget;
+class UVPFullScreenUserWidget;
+class SWindow;
+class FSceneViewport;
+
+#if WITH_EDITOR
+class FLevelEditorViewportClient;
+#endif
+
+UCLASS(BlueprintType, Abstract, EditInlineNew)
+class VCAMCORE_API UVCamOutputProviderBase : public UObject
+{
+	GENERATED_BODY()
+public:
+
+	virtual void BeginDestroy() override;
+
+	// Called when the provider is brought online such as after instantiating or loading a component containing this provider 
+	// Use Initialize for any setup logic that needs to survive between Start / Stop cycles such as spawning transient objects 
+	// 
+	// If bForceInitialization is true then it will force a reinitialization even if the provider was already initialized
+	virtual void Initialize();
+
+	// Called when the provider is being shutdown such as before changing level or on exit
+	virtual void Deinitialize();
+
+	// Called when the provider is Activated
+	virtual void Activate();
+
+	// Called when the provider is Deactivated
+	virtual void Deactivate();
+
+	virtual void Tick(const float DeltaTime);
+
+	// Called to turn on or off this output provider
+	UFUNCTION(BlueprintCallable, Category = "Output")
+	void SetActive(const bool bInActive);
+
+#if WITH_EDITOR
+	virtual void PostEditChangeProperty(FPropertyChangedEvent& PropertyChangedEvent) override;
+#endif
+
+	// Returns if this output provider is currently active or not
+	UFUNCTION(BlueprintPure, Category = "Output")
+	bool IsActive() const { return bIsActive; };
+
+	// Returns if this output provider has been initialized or not
+	UFUNCTION(BlueprintPure, Category = "Output")
+	bool IsInitialized() const { return bInitialized; };
+
+	// Sets the TargetCamera parameter
+	UFUNCTION(BlueprintCallable, Category = "Output")
+	void SetTargetCamera(const UCineCameraComponent* InTargetCamera);
+
+	// Sets the UMG class to be rendered in this output provider
+	UFUNCTION(BlueprintCallable, Category = "Output")
+	void SetUMGClass(const TSubclassOf<UUserWidget> InUMGClass);
+
+	// The UMG class to be rendered in this output provider
+	UPROPERTY(EditAnywhere, BlueprintReadWrite, Category = "Output", meta = (DisplayName="UMG Overlay", DisplayPriority = "2"))
+	TSubclassOf<UUserWidget> UMGClass;
+
+	// Override the default output resolution with a custom value - NOTE you must toggle bIsActive off then back on for this to take effect
+	UPROPERTY(EditAnywhere, BlueprintReadWrite, Category = "Output", meta = (DisplayPriority = "3"))
+	bool bUseOverrideResolution = false;
+
+	// When bUseOverrideResolution is set, use this custom resolution
+	UPROPERTY(EditAnywhere, BlueprintReadWrite, Category = "Output", meta = (DisplayPriority = "4"), meta = (EditCondition = "bUseOverrideResolution", ClampMin = 1))
+	FIntPoint OverrideResolution = { 2048, 1536 };
+
+	UVPFullScreenUserWidget* GetUMGWidget() { return UMGWidget; };
+
+	/** Temporarily disable the output.  Caller must eventually call RestoreOutput. */
+	void SuspendOutput()
+	{
+		if (IsActive())
+		{
+			bWasActive = true;
+			SetActive(false);
+		}
+	}
+
+	/** Restore the output state from previous call to disable output. */
+	void RestoreOutput()
+	{
+		if (bWasActive && !IsActive())
+		{
+			SetActive(true);
+		}
+		bWasActive = false;
+	}
+
+	/** Calls the VCamModifierInterface on the widget if it exists and also requests any child VCam Widgets to reconnect */
+	void NotifyWidgetOfComponentChange() const;
+
+protected:
+	// If set, this output provider will execute every frame
+	UPROPERTY(EditAnywhere, BlueprintReadWrite, Category = "Output", meta = (DisplayPriority = "1"))
+	bool bIsActive = false;
+
+	UPROPERTY(Transient)
+	bool bInitialized = false;
+
+	EVPWidgetDisplayType DisplayType = EVPWidgetDisplayType::PostProcess;
+
+	virtual void CreateUMG();
+
+	void DisplayUMG();
+	void DestroyUMG();
+
+	UVCamOutputProviderBase* GetOtherOutputProviderByIndex(int32 Index) const;
+
+	TSharedPtr<FSceneViewport> GetTargetSceneViewport() const;
+	TWeakPtr<SWindow> GetTargetInputWindow() const;
+
+#if WITH_EDITOR
+	FLevelEditorViewportClient* GetTargetLevelViewportClient() const;
+	TSharedPtr<SLevelViewport> GetTargetLevelViewport() const;
+#endif
+
+	UPROPERTY(Transient)
+	TObjectPtr<UVPFullScreenUserWidget> UMGWidget = nullptr;
+
+private:
+	bool IsOuterComponentEnabled() const;
+
+	TSoftObjectPtr<UCineCameraComponent> TargetCamera;
+	bool bWasActive = false;
+};
diff --git a/Engine/Plugins/Experimental/VirtualCamera/Source/VCamCore/Public/Output/VCamOutputRemoteSession.h b/Engine/Plugins/Experimental/VirtualCamera/Source/VCamCore/Public/Output/VCamOutputRemoteSession.h
new file mode 100644
--- /dev/null
+++ b/Engine/Plugins/Experimental/VirtualCamera/Source/VCamCore/Public/Output/VCamOutputRemoteSession.h
@@ -0,0 +1,67 @@
+// Copyright Epic Games, Inc. All Rights Reserved.
+
+#pragma once
+
+#include "CoreMinimal.h"
+#include "VCamOutputProviderBase.h"
+#include "ImageProviders/RemoteSessionMediaOutput.h"
+#include "IRemoteSessionRole.h"
+#include "RemoteSession.h"
+
+#if WITH_EDITOR
+#include "LevelEditorViewport.h"
+#endif
+
+#include "VCamOutputRemoteSession.generated.h"
+
+UCLASS(meta = (DisplayName = "Unreal Remote Output Provider"))
+class VCAMCORE_API UVCamOutputRemoteSession : public UVCamOutputProviderBase
+{
+	GENERATED_BODY()
+
+public:
+	virtual void Initialize() override;
+	virtual void Deinitialize() override;
+
+	virtual void Activate() override;
+	virtual void Deactivate() override;
+
+	virtual void Tick(const float DeltaTime) override;
+
+#if WITH_EDITOR
+	virtual void PostEditChangeProperty(FPropertyChangedEvent& PropertyChangedEvent) override;
+#endif
+
+	// Network port number - change this only if connecting multiple RemoteSession devices to the same PC
+	UPROPERTY(EditAnywhere, BlueprintReadWrite, Category = "Output")
+	int32 PortNumber = IRemoteSessionModule::kDefaultPort;
+
+	// If using the output from a Composure Output Provider, specify it here
+	UPROPERTY(EditAnywhere, BlueprintReadWrite, Category = "Output")
+	int32 FromComposureOutputProviderIndex = INDEX_NONE;
+
+protected:
+
+	UPROPERTY(Transient)
+	TObjectPtr<URemoteSessionMediaOutput> MediaOutput = nullptr;
+
+	UPROPERTY(Transient)
+	TObjectPtr<URemoteSessionMediaCapture> MediaCapture = nullptr;
+
+private:
+	TSharedPtr<IRemoteSessionUnmanagedRole> RemoteSessionHost;
+
+	FHitResult LastViewportTouchResult;
+
+	bool bUsingDummyUMG = false;
+
+	void CreateRemoteSession();
+	void DestroyRemoteSession();
+
+	void OnRemoteSessionChannelChange(IRemoteSessionRole* Role, TWeakPtr<IRemoteSessionChannel> Channel, ERemoteSessionChannelChange Change);
+	void OnImageChannelCreated(TWeakPtr<IRemoteSessionChannel> Instance);
+	void OnInputChannelCreated(TWeakPtr<IRemoteSessionChannel> Instance);
+	void OnTouchEventOutsideUMG(const FVector2D& InViewportPosition);
+
+	bool DeprojectScreenToWorld(const FVector2D& InScreenPosition, FVector& OutWorldPosition, FVector& OutWorldDirection) const;
+};
diff --git a/Engine/Plugins/Experimental/VirtualCamera/Source/VCamCore/Public/Output/VCamOutputViewport.h b/Engine/Plugins/Experimental/VirtualCamera/Source/VCamCore/Public/Output/VCamOutputViewport.h
new file mode 100644
--- /dev/null
+++ b/Engine/Plugins/Experimental/VirtualCamera/Source/VCamCore/Public/Output/VCamOutputViewport.h
@@ -0,0 +1,20 @@
+// Copyright Epic Games, Inc. All Rights Reserved.
+
+#pragma once
+
+#include "CoreMinimal.h"
+#include "VCamOutputProviderBase.h"
+#include "VCamOutputViewport.generated.h"
+
+UCLASS(meta = (DisplayName = "Viewport Output Provider"))
+class VCAMCORE_API UVCamOutputViewport : public UVCamOutputProviderBase
+{
+	GENERATED_BODY()
+
+public:
+
+protected:
+	virtual void CreateUMG() override;
+
+private:
+};
diff --git a/Engine/Plugins/Experimental/VirtualCamera/Source/VCamCore/Public/VCamComponent.h b/Engine/Plugins/Experimental/VirtualCamera/Source/VCamCore/Public/VCamComponent.h
--- a/Engine/Plugins/Experimental/VirtualCamera/Source/VCamCore/Public/VCamComponent.h
+++ b/Engine/Plugins/Experimental/VirtualCamera/Source/VCamCore/Public/VCamComponent.h
@@ -3,13 +3,16 @@
 #pragma once
 
 #include "CoreMinimal.h"
+
+#include "EVCamTargetViewportID.h"
+#include "Modifier//ModifierStackEntry.h"
+#include "Output/VCamOutputProviderBase.h"
+
 #include "EnhancedActionKeyMapping.h"
 #include "EnhancedInputSubsystemInterface.h"
-#include "VCamTypes.h"
-#include "VCamInputSettings.h"
-#include "Roles/LiveLinkCameraTypes.h"
-#include "VCamOutputProviderBase.h"
 #include "GameplayTagContainer.h"
+#include "Roles/LiveLinkCameraTypes.h"
+#include "VCamInputSettings.h"
 
 #if WITH_EDITOR
 #include "UnrealEdMisc.h"
@@ -32,16 +35,6 @@ class FLevelEditorViewportClient;
 
 DECLARE_DYNAMIC_MULTICAST_DELEGATE_OneParam(FOnComponentReplaced, UVCamComponent*, NewComponent);
 
-UENUM(BlueprintType, meta=(DisplayName = "VCam Target Viewport ID"))
-enum class EVCamTargetViewportID : uint8
-{
-	CurrentlySelected = 0,
-	Viewport1 = 1,
-	Viewport2 = 2,
-	Viewport3 = 3,
-	Viewport4 = 4
-};
-
 UCLASS(Blueprintable, ClassGroup = VCam, HideCategories=(Mobility), meta=(BlueprintSpawnableComponent))
 class VCAMCORE_API UVCamComponent : public USceneComponent
 {
diff --git a/Engine/Plugins/Experimental/VirtualCamera/Source/VCamCore/Public/VCamModifier.h b/Engine/Plugins/Experimental/VirtualCamera/Source/VCamCore/Public/VCamModifier.h
deleted file mode 100644
diff --git a/Engine/Plugins/Experimental/VirtualCamera/Source/VCamCore/Public/VCamModifierContext.h b/Engine/Plugins/Experimental/VirtualCamera/Source/VCamCore/Public/VCamModifierContext.h
deleted file mode 100644
diff --git a/Engine/Plugins/Experimental/VirtualCamera/Source/VCamCore/Public/VCamModifierInterface.h b/Engine/Plugins/Experimental/VirtualCamera/Source/VCamCore/Public/VCamModifierInterface.h
deleted file mode 100644
diff --git a/Engine/Plugins/Experimental/VirtualCamera/Source/VCamCore/Public/VCamOutputComposure.h b/Engine/Plugins/Experimental/VirtualCamera/Source/VCamCore/Public/VCamOutputComposure.h
deleted file mode 100644
diff --git a/Engine/Plugins/Experimental/VirtualCamera/Source/VCamCore/Public/VCamOutputMediaOutput.h b/Engine/Plugins/Experimental/VirtualCamera/Source/VCamCore/Public/VCamOutputMediaOutput.h
deleted file mode 100644
diff --git a/Engine/Plugins/Experimental/VirtualCamera/Source/VCamCore/Public/VCamOutputProviderBase.h b/Engine/Plugins/Experimental/VirtualCamera/Source/VCamCore/Public/VCamOutputProviderBase.h
deleted file mode 100644
diff --git a/Engine/Plugins/Experimental/VirtualCamera/Source/VCamCore/Public/VCamOutputRemoteSession.h b/Engine/Plugins/Experimental/VirtualCamera/Source/VCamCore/Public/VCamOutputRemoteSession.h
deleted file mode 100644
diff --git a/Engine/Plugins/Experimental/VirtualCamera/Source/VCamCore/Public/VCamOutputViewport.h b/Engine/Plugins/Experimental/VirtualCamera/Source/VCamCore/Public/VCamOutputViewport.h
deleted file mode 100644
diff --git a/Engine/Plugins/Experimental/VirtualCamera/Source/VCamCore/Public/VCamTypes.h b/Engine/Plugins/Experimental/VirtualCamera/Source/VCamCore/Public/VCamTypes.h
deleted file mode 100644
diff --git a/Engine/Plugins/Experimental/VirtualCamera/Source/VCamCoreEditor/Private/Factories/VCamModifierFactory.cpp b/Engine/Plugins/Experimental/VirtualCamera/Source/VCamCoreEditor/Private/Factories/VCamModifierFactory.cpp
--- a/Engine/Plugins/Experimental/VirtualCamera/Source/VCamCoreEditor/Private/Factories/VCamModifierFactory.cpp
+++ b/Engine/Plugins/Experimental/VirtualCamera/Source/VCamCoreEditor/Private/Factories/VCamModifierFactory.cpp
@@ -1,7 +1,8 @@
 // Copyright Epic Games, Inc. All Rights Reserved.
 
 #include "Factories/VCamModifierFactory.h"
-#include "VCamModifier.h"
+
+#include "Modifier/VCamModifier.h"
 
 #include "AssetToolsModule.h"
 #include "Kismet2/BlueprintEditorUtils.h"
diff --git a/Engine/Plugins/Experimental/VirtualCamera/Source/VCamCoreEditor/Private/VCamCoreEditorModule.cpp b/Engine/Plugins/Experimental/VirtualCamera/Source/VCamCoreEditor/Private/VCamCoreEditorModule.cpp
--- a/Engine/Plugins/Experimental/VirtualCamera/Source/VCamCoreEditor/Private/VCamCoreEditorModule.cpp
+++ b/Engine/Plugins/Experimental/VirtualCamera/Source/VCamCoreEditor/Private/VCamCoreEditorModule.cpp
@@ -10,7 +10,7 @@
 #include "UI/Switcher/WidgetConnectionConfig.h"
 #include "VCamBaseActor.h"
 #include "VCamComponent.h"
-#include "VCamModifier.h"
+#include "Modifier/VCamModifier.h"
 
 #include "Interfaces/IPluginManager.h"
 #include "Kismet2/KismetEditorUtilities.h"
