Search Actor does not show result for VCamActor

diff --git a/Engine/Plugins/Experimental/VirtualCamera/Source/VirtualCameraEditor/Private/VirtualCameraEditorModule.cpp b/Engine/Plugins/Experimental/VirtualCamera/Source/VirtualCameraEditor/Private/VirtualCameraEditorModule.cpp
--- a/Engine/Plugins/Experimental/VirtualCamera/Source/VirtualCameraEditor/Private/VirtualCameraEditorModule.cpp
+++ b/Engine/Plugins/Experimental/VirtualCamera/Source/VirtualCameraEditor/Private/VirtualCameraEditorModule.cpp
@@ -2,23 +2,23 @@
 
 #include "VirtualCameraEditorModule.h"
 
+#include "ActorFactories/ActorFactoryBlueprint.h"
+#include "Filters/CustomClassFilterData.h"
+#include "IPlacementModeModule.h"
 #include "ISettingsModule.h"
 #include "ISettingsSection.h"
-#include "Modules/ModuleManager.h"
-#include "VirtualCameraUserSettings.h"
-#include "IPlacementModeModule.h"
-#include "ActorFactories/ActorFactoryBlueprint.h"
 #include "IVPUtilitiesEditorModule.h"
-#include "SimpleVirtualCamera.h"
-#include "VirtualCameraActor.h"
 #include "LevelEditor.h"
 #include "LevelEditorOutlinerSettings.h"
-#include "Filters/CustomClassFilterData.h"
-
-
+#include "Modules/ModuleManager.h"
+#include "SimpleVirtualCamera.h"
+#include "VirtualCameraActor.h"
+#include "VirtualCameraUserSettings.h"
 
 #define LOCTEXT_NAMESPACE "FVirtualCameraEditorModule"
 
+DEFINE_LOG_CATEGORY_STATIC(LogVirtualCameraEditor, Log, All);
+
 class FVirtualCameraEditorModule : public IModuleInterface
 {
 public:
@@ -57,28 +57,33 @@ public:
 
 	void RegisterPlacementModeItems()
 	{
-		if (GEditor)
+		if (const FPlacementCategoryInfo* Info = IVPUtilitiesEditorModule::Get().GetVirtualProductionPlacementCategoryInfo()
+			; Info && GEditor)
 		{
-			if (const FPlacementCategoryInfo* Info = IVPUtilitiesEditorModule::Get().GetVirtualProductionPlacementCategoryInfo())
-			{
-				FAssetData VirtualCamera2ActorAssetData(
-					TEXT("/VirtualCamera/VCamActor"),
-					TEXT("/VirtualCamera"),
-					TEXT("VCamActor"),
-					FTopLevelAssetPath(TEXT("/Script/Engine"), TEXT("Blueprint"))
-				);
+			FAssetData VirtualCamera2ActorAssetData(
+				TEXT("/VirtualCamera/VCamActor"),
+				TEXT("/VirtualCamera"),
+				TEXT("VCamActor"),
+				FTopLevelAssetPath(TEXT("/Script/Engine"), TEXT("Blueprint"))
+			);
+
+			// Makes it appear in the VP category ...
+			IPlacementModeModule::Get().RegisterPlaceableItem(Info->UniqueHandle, MakeShared<FPlaceableItem>(
+				*UActorFactoryBlueprint::StaticClass(),
+				VirtualCamera2ActorAssetData,
+				FName("ClassThumbnail.CameraActor"),
+				FName("ClassIcon.CameraActor"),
+				TOptional<FLinearColor>(),
+				TOptional<int32>(),
+				NSLOCTEXT("PlacementMode", "VCam Actor", "VCam Actor")
+				)
+			);
 
-				// register the full-fat camera
-				IPlacementModeModule::Get().RegisterPlaceableItem(Info->UniqueHandle, MakeShared<FPlaceableItem>(
-					*UActorFactoryBlueprint::StaticClass(),
-					VirtualCamera2ActorAssetData,
-					FName("ClassThumbnail.CameraActor"),
-					FName("ClassIcon.CameraActor"),
-					TOptional<FLinearColor>(),
-					TOptional<int32>(),
-					NSLOCTEXT("PlacementMode", "VCam Actor", "VCam Actor")
-				));
-			}
+			// ... but if you search for it by text this is needed to make it show up (without having the user load it manually).
+			// The search filters everything in the FBuiltInPlacementCategories::AllClasses category only;
+			// it contains 1. loaded BP classes and 2. specialized actor factories. This manual load adds it to case 1. 
+			const UClass* VCamActorBlueprintClass = LoadClass<UClass>(nullptr, TEXT("/VirtualCamera/VCamActor.VCamActor_C"));
+			UE_CLOG(VCamActorBlueprintClass == nullptr, LogVirtualCameraEditor, Warning, TEXT("Failed to load '/VirtualCamera/VCamActor.VCamActor_C'. Has the Blueprint been moved?"));
 		}
 	}
 
