Fix copy-paste error when creating new module

diff --git a/Engine/Plugins/Experimental/VirtualCamera/Source/VCamBlueprintNodes/Private/VCamBlueprintNodes.cpp b/Engine/Plugins/Experimental/VirtualCamera/Source/VCamBlueprintNodes/Private/VCamBlueprintNodes.cpp
--- a/Engine/Plugins/Experimental/VirtualCamera/Source/VCamBlueprintNodes/Private/VCamBlueprintNodes.cpp
+++ b/Engine/Plugins/Experimental/VirtualCamera/Source/VCamBlueprintNodes/Private/VCamBlueprintNodes.cpp
@@ -10,15 +10,11 @@ namespace UE::VCamBlueprintNodes::Private
 	public:
 
 		virtual void StartupModule() override
-		{
-			
-		}
+		{}
 	
 		virtual void ShutdownModule() override
-		{
-			
-		}
+		{}
 	};
 }
 
-IMPLEMENT_MODULE(UE::VCamBlueprintNodes::Private::FVCamBlueprintNodesModule, VCamCoreEditor);
\ No newline at end of file
+IMPLEMENT_MODULE(UE::VCamBlueprintNodes::Private::FVCamBlueprintNodesModule, VCamBlueprintNodes);
\ No newline at end of file
