Fix not being able to load VCamActor_C due to incorrect parameter being passed to LoadClass

diff --git a/Engine/Plugins/Experimental/VirtualCamera/Source/VirtualCameraEditor/Private/VirtualCameraEditorModule.cpp b/Engine/Plugins/Experimental/VirtualCamera/Source/VirtualCameraEditor/Private/VirtualCameraEditorModule.cpp
--- a/Engine/Plugins/Experimental/VirtualCamera/Source/VirtualCameraEditor/Private/VirtualCameraEditorModule.cpp
+++ b/Engine/Plugins/Experimental/VirtualCamera/Source/VirtualCameraEditor/Private/VirtualCameraEditorModule.cpp
@@ -82,7 +82,7 @@ public:
 			// ... but if you search for it by text this is needed to make it show up (without having the user load it manually).
 			// The search filters everything in the FBuiltInPlacementCategories::AllClasses category only;
 			// it contains 1. loaded BP classes and 2. specialized actor factories. This manual load adds it to case 1. 
-			const UClass* VCamActorBlueprintClass = LoadClass<UClass>(nullptr, TEXT("/VirtualCamera/VCamActor.VCamActor_C"));
+			const UClass* VCamActorBlueprintClass = LoadClass<UObject>(nullptr, TEXT("/VirtualCamera/VCamActor.VCamActor_C"));
 			UE_CLOG(VCamActorBlueprintClass == nullptr, LogVirtualCameraEditor, Warning, TEXT("Failed to load '/VirtualCamera/VCamActor.VCamActor_C'. Has the Blueprint been moved?"));
 		}
 	}
