Fix CIS

diff --git a/Engine/Plugins/Experimental/VirtualCamera/Source/VCamCore/Public/VCamComponent.h b/Engine/Plugins/Experimental/VirtualCamera/Source/VCamCore/Public/VCamComponent.h
--- a/Engine/Plugins/Experimental/VirtualCamera/Source/VCamCore/Public/VCamComponent.h
+++ b/Engine/Plugins/Experimental/VirtualCamera/Source/VCamCore/Public/VCamComponent.h
@@ -5,7 +5,7 @@
 #include "CoreMinimal.h"
 
 #include "EVCamTargetViewportID.h"
-#include "Modifier//ModifierStackEntry.h"
+#include "Modifier/ModifierStackEntry.h"
 #include "Output/VCamOutputProviderBase.h"
 
 #include "EnhancedActionKeyMapping.h"
