Add way for output providers to specify their target viewport

diff --git a/Engine/Plugins/Experimental/VirtualCamera/Source/PixelStreamingVCam/Private/VCamPixelStreamingSession.cpp b/Engine/Plugins/Experimental/VirtualCamera/Source/PixelStreamingVCam/Private/VCamPixelStreamingSession.cpp
--- a/Engine/Plugins/Experimental/VirtualCamera/Source/PixelStreamingVCam/Private/VCamPixelStreamingSession.cpp
+++ b/Engine/Plugins/Experimental/VirtualCamera/Source/PixelStreamingVCam/Private/VCamPixelStreamingSession.cpp
@@ -61,15 +61,15 @@ void UVCamPixelStreamingSession::Activate()
 		UVCamComponent* VCamComponent = GetTypedOuter<UVCamComponent>();
 		if (bAutoSetLiveLinkSubject && IsValid(VCamComponent))
 		{
-			VCamComponent->LiveLinkSubject = GetFName();
+			VCamComponent->SetLiveLinkSubobject(GetFName());
 		}
 	}
 
 	// If we don't have a UMG assigned, we still need to create an empty 'dummy' UMG in order to properly route the input back from the RemoteSession device
-	if (UMGClass == nullptr)
+	if (!GetUMGClass())
 	{
 		bUsingDummyUMG = true;
-		UMGClass = UE::VCamPixelStreamingSession::Private::EmptyUMGSoftClassPath.TryLoadClass<UUserWidget>();
+		SetUMGClass(UE::VCamPixelStreamingSession::Private::EmptyUMGSoftClassPath.TryLoadClass<UUserWidget>());
 	}
 
 	if (MediaOutput == nullptr)
@@ -336,7 +336,7 @@ void UVCamPixelStreamingSession::Deactivate()
 	Super::Deactivate();
 	if (bUsingDummyUMG)
 	{
-		UMGClass = nullptr;
+		SetUMGClass(nullptr);
 		bUsingDummyUMG = false;
 	}
 
@@ -374,16 +374,6 @@ void UVCamPixelStreamingSession::PostEditChangeProperty(FPropertyChangedEvent& P
 }
 #endif
 
-TWeakPtr<SWindow> UVCamPixelStreamingSession::GetTargetInputWindow() const
-{
-	TWeakPtr<SWindow> InputWindow;
-	if (UVCamComponent* OuterComponent = GetTypedOuter<UVCamComponent>())
-	{
-		InputWindow = OuterComponent->GetTargetInputWindow();
-	}
-	return InputWindow;
-}
-
 void UVCamPixelStreamingSession::UpdateOverrideResolution(bool bApplyOverride) const
 {
 	if (const TSharedPtr<FSceneViewport> TargetSceneViewport = GetTargetSceneViewport())
diff --git a/Engine/Plugins/Experimental/VirtualCamera/Source/PixelStreamingVCam/Public/VCamPixelStreamingSession.h b/Engine/Plugins/Experimental/VirtualCamera/Source/PixelStreamingVCam/Public/VCamPixelStreamingSession.h
--- a/Engine/Plugins/Experimental/VirtualCamera/Source/PixelStreamingVCam/Public/VCamPixelStreamingSession.h
+++ b/Engine/Plugins/Experimental/VirtualCamera/Source/PixelStreamingVCam/Public/VCamPixelStreamingSession.h
@@ -56,7 +56,6 @@ private:
 	void SetupCustomInputHandling();
 	void OnCaptureStateChanged();
 	void OnARKitTransformReceived(FPixelStreamingPlayerId PlayerId, uint8 Type, TArray<uint8> Data);
-	TWeakPtr<SWindow> GetTargetInputWindow() const;
 	
 	/**
 	 * Applies or removes a fixed resolution to a viewport.
diff --git a/Engine/Plugins/Experimental/VirtualCamera/Source/VCamCore/Private/Output/VCamOutputProviderBase.cpp b/Engine/Plugins/Experimental/VirtualCamera/Source/VCamCore/Private/Output/VCamOutputProviderBase.cpp
--- a/Engine/Plugins/Experimental/VirtualCamera/Source/VCamCore/Private/Output/VCamOutputProviderBase.cpp
+++ b/Engine/Plugins/Experimental/VirtualCamera/Source/VCamCore/Private/Output/VCamOutputProviderBase.cpp
@@ -5,19 +5,33 @@
 #include "Modifier/VCamModifierInterface.h"
 #include "VCamComponent.h"
 #include "UI/VCamWidget.h"
+#include "Util/LevelViewportUtils.h"
+#include "VCamCoreCustomVersion.h"
 
-#include "Engine/Engine.h"
-#include "UObject/UObjectBaseUtility.h"
 #include "Blueprint/UserWidget.h"
 #include "Blueprint/WidgetTree.h"
+#include "Engine/Engine.h"
+#include "Framework/Application/SlateApplication.h"
+#include "UObject/UObjectBaseUtility.h"
 
 #if WITH_EDITOR
 #include "Editor.h"
+#include "IAssetViewport.h"
+#include "LevelEditorViewport.h"
+#include "SLevelViewport.h"
+#include "SEditorViewport.h"
 #include "UnrealClient.h"
+#else
+#include "Engine/GameEngine.h"
 #endif
 
 DEFINE_LOG_CATEGORY(LogVCamOutputProvider);
 
+namespace UE::VCamCore::Private
+{
+	static const FName LevelEditorName(TEXT("LevelEditor"));
+}
+
 void UVCamOutputProviderBase::BeginDestroy()
 {
 	Deinitialize();
@@ -110,15 +124,9 @@ bool UVCamOutputProviderBase::IsOuterComponentEnabled() const
 void UVCamOutputProviderBase::SetTargetCamera(const UCineCameraComponent* InTargetCamera)
 {
 	TargetCamera = InTargetCamera;
-
 	NotifyWidgetOfComponentChange();
 }
 
-void UVCamOutputProviderBase::SetUMGClass(const TSubclassOf<UUserWidget> InUMGClass)
-{
-	UMGClass = InUMGClass;
-}
-
 void UVCamOutputProviderBase::CreateUMG()
 {
 	if (!UMGClass)
@@ -234,6 +242,21 @@ void UVCamOutputProviderBase::NotifyWidgetOfComponentChange() const
 	}
 }
 
+void UVCamOutputProviderBase::Serialize(FArchive& Ar)
+{
+	using namespace UE::VCamCore;
+	Super::Serialize(Ar);
+	Ar.UsingCustomVersion(FVCamCoreCustomVersion::GUID);
+	
+	if (Ar.IsLoading() && Ar.CustomVer(FVCamCoreCustomVersion::GUID) < FVCamCoreCustomVersion::MoveTargetViewportFromComponentToOutput)
+	{
+		UVCamComponent* OuterComponent = GetTypedOuter<UVCamComponent>();
+		TargetViewport = OuterComponent
+			? OuterComponent->TargetViewport_DEPRECATED
+			: TargetViewport;
+	}
+}
+
 UVCamOutputProviderBase* UVCamOutputProviderBase::GetOtherOutputProviderByIndex(int32 Index) const
 {
 	if (Index > INDEX_NONE)
@@ -254,15 +277,53 @@ UVCamOutputProviderBase* UVCamOutputProviderBase::GetOtherOutputProviderByIndex(
 	return nullptr;
 }
 
-
 TSharedPtr<FSceneViewport> UVCamOutputProviderBase::GetTargetSceneViewport() const
 {
 	TSharedPtr<FSceneViewport> SceneViewport;
 
-	if (UVCamComponent* OuterComponent = GetTypedOuter<UVCamComponent>())
+#if WITH_EDITOR
+	if (GIsEditor)
 	{
-		SceneViewport = OuterComponent->GetTargetSceneViewport();
+		for (const FWorldContext& Context : GEngine->GetWorldContexts())
+		{
+			if (Context.WorldType == EWorldType::PIE)
+			{
+				FSlatePlayInEditorInfo* SlatePlayInEditorSession = GEditor->SlatePlayInEditorMap.Find(Context.ContextHandle);
+				if (SlatePlayInEditorSession)
+				{
+					if (SlatePlayInEditorSession->DestinationSlateViewport.IsValid())
+					{
+						TSharedPtr<IAssetViewport> DestinationLevelViewport = SlatePlayInEditorSession->DestinationSlateViewport.Pin();
+						SceneViewport = DestinationLevelViewport->GetSharedActiveViewport();
+					}
+					else if (SlatePlayInEditorSession->SlatePlayInEditorWindowViewport.IsValid())
+					{
+						SceneViewport = SlatePlayInEditorSession->SlatePlayInEditorWindowViewport;
+					}
+
+					// If PIE is active always choose it
+					break;
+				}
+			}
+			else if (Context.WorldType == EWorldType::Editor)
+			{
+				if (FLevelEditorViewportClient* LevelViewportClient = GetTargetLevelViewportClient())
+				{
+					TSharedPtr<SEditorViewport> ViewportWidget = LevelViewportClient->GetEditorViewportWidget();
+					if (ViewportWidget.IsValid())
+					{
+						SceneViewport = ViewportWidget->GetSceneViewport();
+					}
+				}
+			}
+		}
 	}
+#else
+	if (UGameEngine* GameEngine = Cast<UGameEngine>(GEngine))
+	{
+		SceneViewport = GameEngine->SceneViewport;
+	}
+#endif
 
 	return SceneViewport;
 }
@@ -271,37 +332,66 @@ TWeakPtr<SWindow> UVCamOutputProviderBase::GetTargetInputWindow() const
 {
 	TWeakPtr<SWindow> InputWindow;
 
-	if (UVCamComponent* OuterComponent = GetTypedOuter<UVCamComponent>())
+#if WITH_EDITOR
+	if (GIsEditor)
 	{
-		InputWindow = OuterComponent->GetTargetInputWindow();
+		for (const FWorldContext& Context : GEngine->GetWorldContexts())
+		{
+			if (Context.WorldType == EWorldType::PIE)
+			{
+				FSlatePlayInEditorInfo* SlatePlayInEditorSession = GEditor->SlatePlayInEditorMap.Find(Context.ContextHandle);
+				if (SlatePlayInEditorSession)
+				{
+					if (SlatePlayInEditorSession->DestinationSlateViewport.IsValid())
+					{
+						TSharedPtr<IAssetViewport> DestinationLevelViewport = SlatePlayInEditorSession->DestinationSlateViewport.Pin();
+						InputWindow = FSlateApplication::Get().FindWidgetWindow(DestinationLevelViewport->AsWidget());
+					}
+					else if (SlatePlayInEditorSession->SlatePlayInEditorWindowViewport.IsValid())
+					{
+						InputWindow = SlatePlayInEditorSession->SlatePlayInEditorWindow;
+					}
+
+					// If PIE is active always choose it
+					break;
+				}
+			}
+			else if (Context.WorldType == EWorldType::Editor)
+			{
+				if (FLevelEditorViewportClient* LevelViewportClient = GetTargetLevelViewportClient())
+				{
+					TSharedPtr<SEditorViewport> ViewportWidget = LevelViewportClient->GetEditorViewportWidget();
+					if (ViewportWidget.IsValid())
+					{
+						InputWindow = FSlateApplication::Get().FindWidgetWindow(ViewportWidget.ToSharedRef());
+					}
+				}
+			}
+		}
+	}
+#else
+	if (UGameEngine* GameEngine = Cast<UGameEngine>(GEngine))
+	{
+		InputWindow = GameEngine->GameViewportWindow;
 	}
+#endif
 
 	return InputWindow;
 }
 
 #if WITH_EDITOR
+
 FLevelEditorViewportClient* UVCamOutputProviderBase::GetTargetLevelViewportClient() const
 {
-	FLevelEditorViewportClient* ViewportClient = nullptr;
-
-	if (UVCamComponent* OuterComponent = GetTypedOuter<UVCamComponent>())
-	{
-		ViewportClient = OuterComponent->GetTargetLevelViewportClient();
-	}
-
-	return ViewportClient;
+	const TSharedPtr<SLevelViewport> LevelViewport = GetTargetLevelViewport();
+	return LevelViewport
+		? &LevelViewport->GetLevelViewportClient()
+		: nullptr;
 }
 
 TSharedPtr<SLevelViewport> UVCamOutputProviderBase::GetTargetLevelViewport() const
 {
-	TSharedPtr<SLevelViewport> LevelViewport;
-
-	if (UVCamComponent* OuterComponent = GetTypedOuter<UVCamComponent>())
-	{
-		LevelViewport = OuterComponent->GetTargetLevelViewport();
-	}
-
-	return LevelViewport;
+	return UE::VCamCore::LevelViewportUtils::Private::GetLevelViewport(TargetViewport);
 }
 
 void UVCamOutputProviderBase::PostEditChangeProperty(FPropertyChangedEvent& PropertyChangedEvent)
diff --git a/Engine/Plugins/Experimental/VirtualCamera/Source/VCamCore/Private/Output/VCamOutputRemoteSession.cpp b/Engine/Plugins/Experimental/VirtualCamera/Source/VCamCore/Private/Output/VCamOutputRemoteSession.cpp
--- a/Engine/Plugins/Experimental/VirtualCamera/Source/VCamCore/Private/Output/VCamOutputRemoteSession.cpp
+++ b/Engine/Plugins/Experimental/VirtualCamera/Source/VCamCore/Private/Output/VCamOutputRemoteSession.cpp
@@ -45,10 +45,10 @@ void UVCamOutputRemoteSession::Deinitialize()
 void UVCamOutputRemoteSession::Activate()
 {
 	// If we don't have a UMG assigned, we still need to create an empty 'dummy' UMG in order to properly route the input back from the RemoteSession device
-	if (UMGClass == nullptr)
+	if (!GetUMGClass())
 	{
 		bUsingDummyUMG = true;
-		UMGClass = UE::VCamOutputRemoteSession::Private::EmptyUMGSoftClassPath.TryLoadClass<UUserWidget>();
+		SetUMGClass(UE::VCamOutputRemoteSession::Private::EmptyUMGSoftClassPath.TryLoadClass<UUserWidget>());
 	}
 
 	CreateRemoteSession();
@@ -64,7 +64,7 @@ void UVCamOutputRemoteSession::Deactivate()
 
 	if (bUsingDummyUMG)
 	{
-		UMGClass = nullptr;
+		SetUMGClass(nullptr);
 		bUsingDummyUMG = false;
 	}
 }
@@ -234,7 +234,7 @@ void UVCamOutputRemoteSession::OnInputChannelCreated(TWeakPtr<IRemoteSessionChan
 	if (InputChannel)
 	{
 		// If we have a UMG, then use it
-		if (UMGClass && UMGWidget)
+		if (GetUMGClass() && UMGWidget) 
 		{
 			TSharedPtr<SVirtualWindow> InputWindow;
 
diff --git a/Engine/Plugins/Experimental/VirtualCamera/Source/VCamCore/Private/Util/LevelViewportUtils.cpp b/Engine/Plugins/Experimental/VirtualCamera/Source/VCamCore/Private/Util/LevelViewportUtils.cpp
new file mode 100644
--- /dev/null
+++ b/Engine/Plugins/Experimental/VirtualCamera/Source/VCamCore/Private/Util/LevelViewportUtils.cpp
@@ -0,0 +1,182 @@
+// Copyright Epic Games, Inc. All Rights Reserved.
+
+#include "LevelViewportUtils.h"
+
+#include "EVCamTargetViewportID.h"
+#include "GameFramework/PlayerController.h"
+#include "Output/VCamOutputProviderBase.h"
+#include "Util/VCamViewportLocker.h"
+
+
+#if WITH_EDITOR
+#include "Modules/ModuleManager.h"
+#include "LevelEditor.h"
+#include "LevelEditorViewport.h"
+#include "SLevelViewport.h"
+#endif
+
+namespace UE::VCamCore::LevelViewportUtils::Private
+{
+	namespace Locking
+	{
+#if WITH_EDITOR
+		static void UpdateLockStateForEditor(FVCamViewportLockState& ViewportLockState, EVCamTargetViewportID ViewportID, bool bNewLockState, AActor* ActorToLockWith = nullptr)
+		{
+			TSharedPtr<SLevelViewport> Viewport = GetLevelViewport(ViewportID);
+			FLevelEditorViewportClient* LevelViewportClient = Viewport.IsValid()
+				? &Viewport->GetLevelViewportClient()
+				: nullptr;
+			const bool bNeedsLock = !ViewportLockState.bIsLockedToViewport && bNewLockState;
+			if (LevelViewportClient)
+			{
+				if (bNeedsLock)
+				{
+					ViewportLockState.Backup_ActorLock = LevelViewportClient->GetActiveActorLock();
+					LevelViewportClient->SetActorLock(ActorToLockWith);
+					// If bLockedCameraView is not true then the viewport is locked to the actor's transform and not the camera component
+					LevelViewportClient->bLockedCameraView = true;
+					ViewportLockState.bIsLockedToViewport = true;
+				}
+				else if (ViewportLockState.bIsLockedToViewport && !bNewLockState
+					&& ViewportLockState.Backup_ActorLock.IsValid())
+				{
+					LevelViewportClient->SetActorLock(ViewportLockState.Backup_ActorLock.Get());
+					ViewportLockState.Backup_ActorLock = nullptr;
+					// If bLockedCameraView is not true then the viewport is locked to the actor's transform and not the camera component
+					LevelViewportClient->bLockedCameraView = true;
+					ViewportLockState.bIsLockedToViewport = false;
+				}
+				else if (ViewportLockState.bIsLockedToViewport && !bNewLockState)
+				{
+					LevelViewportClient->SetActorLock(nullptr);
+					ViewportLockState.bIsLockedToViewport = false;
+				}
+			}
+		}
+#endif
+
+		static void UpdateLockStateForGame(FVCamViewportLockState& ViewportLockState, const FWorldContext& Context, bool bNewLockState, AActor* ActorToLockWith = nullptr)
+		{
+			UWorld* ActorWorld = Context.World();
+			if (ActorWorld && ActorWorld->GetGameInstance())
+			{
+				APlayerController* PlayerController = ActorWorld->GetGameInstance()->GetFirstLocalPlayerController(ActorWorld);
+				if (PlayerController)
+				{
+					const bool bNeedsLock = !ViewportLockState.bIsLockedToViewport && bNewLockState;
+					if (bNeedsLock)
+					{
+						ViewportLockState.Backup_ViewTarget = PlayerController->GetViewTarget();
+						PlayerController->SetViewTarget(ActorToLockWith);
+						ViewportLockState.bIsLockedToViewport = true;
+					}
+					else if (ViewportLockState.bIsLockedToViewport && !bNewLockState
+						&& ViewportLockState.Backup_ViewTarget.IsValid())
+					{
+						PlayerController->SetViewTarget(ViewportLockState.Backup_ViewTarget.Get());
+						ViewportLockState.Backup_ViewTarget = nullptr;
+						ViewportLockState.bIsLockedToViewport = false;
+					}
+					else if (ViewportLockState.bIsLockedToViewport && !bNewLockState)
+					{
+						PlayerController->SetViewTarget(nullptr);
+						ViewportLockState.bIsLockedToViewport = false;
+					}
+				}
+			}
+		}
+		
+		static void UpdateLockState(FVCamViewportLockState& ViewportLockState, EVCamTargetViewportID ViewportID, bool bForcfullySetUnlocked, AActor* ActorToLockWith = nullptr)
+		{
+			const bool bNewLockState = ViewportLockState.bLockViewportToCamera & !bForcfullySetUnlocked;
+			for (const FWorldContext& Context : GEngine->GetWorldContexts())
+			{
+#if WITH_EDITOR
+				if (Context.WorldType == EWorldType::Editor)
+				{
+					UpdateLockStateForEditor(ViewportLockState, ViewportID, bNewLockState, ActorToLockWith);
+				}
+				else
+#endif
+				{
+					UpdateLockStateForGame(ViewportLockState, Context, bNewLockState, ActorToLockWith);
+				}
+			}
+		}
+	}
+
+	void UpdateViewportLocksFromOutputs(TArray<TObjectPtr<UVCamOutputProviderBase>> OutputProviders, FVCamViewportLocker& LockData, AActor* ActorToLockWith)
+	{
+		TSet<EVCamTargetViewportID> Viewports;
+		Algo::TransformIf(OutputProviders, Viewports,
+			[](TObjectPtr<UVCamOutputProviderBase> Output){ return Output && Output->IsActive(); },
+			[](TObjectPtr<UVCamOutputProviderBase> Output) { return Output->GetTargetViewport(); }
+			);
+		
+		check(LockData.Locks.Num() == 4);
+		for (TPair<EVCamTargetViewportID, FVCamViewportLockState>& ViewportData : LockData.Locks)
+		{
+			const bool bShouldUnlock = !Viewports.Contains(ViewportData.Key);
+			Locking::UpdateLockState(ViewportData.Value, ViewportData.Key, bShouldUnlock, ActorToLockWith);
+		}
+	}
+
+	void UnlockAllViewports(FVCamViewportLocker& LockData)
+	{
+		for (TPair<EVCamTargetViewportID, FVCamViewportLockState>& ViewportData : LockData.Locks)
+		{
+			Locking::UpdateLockState(ViewportData.Value, ViewportData.Key, false);
+		}
+	}
+
+#if WITH_EDITOR
+	TSharedPtr<SLevelViewport> GetLevelViewport(EVCamTargetViewportID TargetViewport)
+	{
+		TSharedPtr<SLevelViewport> OutLevelViewport = nullptr;
+
+		if (TargetViewport == EVCamTargetViewportID::CurrentlySelected)
+		{
+			FLevelEditorModule* LevelEditorModule = FModuleManager::GetModulePtr<FLevelEditorModule>(TEXT("LevelEditor"));
+			return LevelEditorModule
+				? LevelEditorModule->GetFirstActiveLevelViewport()
+				: nullptr;;
+		}
+
+		if (!GEditor)
+		{
+			return nullptr;
+		}
+	
+		for (FLevelEditorViewportClient* Client : GEditor->GetLevelViewportClients())
+		{
+			// We only care about the fully rendered 3D viewport...seems like there should be a better way to check for this
+			if (Client->IsOrtho())
+			{
+				continue;
+			}
+
+			TSharedPtr<SLevelViewport> LevelViewport = StaticCastSharedPtr<SLevelViewport>(Client->GetEditorViewportWidget());
+			if (!LevelViewport.IsValid())
+			{
+				continue;
+			}
+		
+			const FString WantedViewportString = GetConfigKeyFor(TargetViewport);
+			const FString ViewportConfigKey = LevelViewport->GetConfigKey().ToString();
+			if (ViewportConfigKey.Contains(*WantedViewportString, ESearchCase::CaseSensitive, ESearchDir::FromStart))
+			{
+				return LevelViewport;
+			}
+		}
+
+		return OutLevelViewport;
+	}
+	
+	FString GetConfigKeyFor(EVCamTargetViewportID TargetViewport)
+	{
+		return TargetViewport == EVCamTargetViewportID::CurrentlySelected
+			? FName(EName::None).ToString()
+			: FString::Printf(TEXT("Viewport %d.Viewport"), (int32)TargetViewport);
+	}
+#endif
+}
diff --git a/Engine/Plugins/Experimental/VirtualCamera/Source/VCamCore/Private/Util/LevelViewportUtils.h b/Engine/Plugins/Experimental/VirtualCamera/Source/VCamCore/Private/Util/LevelViewportUtils.h
new file mode 100644
--- /dev/null
+++ b/Engine/Plugins/Experimental/VirtualCamera/Source/VCamCore/Private/Util/LevelViewportUtils.h
@@ -0,0 +1,27 @@
+// Copyright Epic Games, Inc. All Rights Reserved.
+
+#pragma once
+
+#include "CoreMinimal.h"
+
+class AActor;
+class SLevelViewport;
+class UVCamOutputProviderBase;
+enum class EVCamTargetViewportID : uint8;
+struct FVCamViewportLocker;
+
+namespace UE::VCamCore::LevelViewportUtils::Private
+{
+	/** Locks or unlocks all viewports which are used by OutputProviders depending on the corresponding FVCamViewportLockState::bLockViewportToCamera flag. All other viewports will be unlocked. */
+	void UpdateViewportLocksFromOutputs(TArray<TObjectPtr<UVCamOutputProviderBase>> OutputProviders, FVCamViewportLocker& LockData, AActor* ActorToLockWith);
+	/** Sets all viewports to be unlocked. */
+	void UnlockAllViewports(FVCamViewportLocker& LockData);
+
+#if WITH_EDITOR
+	/** Gets the level viewport identified by TargetViewport */
+	TSharedPtr<SLevelViewport> GetLevelViewport(EVCamTargetViewportID TargetViewport);
+	
+	/** Gets the what SLevelViewport::GetConfigKey would return or None in the case of EVCamTargetViewportID::CurrentlySelected */
+	FString GetConfigKeyFor(EVCamTargetViewportID TargetViewport);
+#endif
+};
diff --git a/Engine/Plugins/Experimental/VirtualCamera/Source/VCamCore/Private/VCamComponent.cpp b/Engine/Plugins/Experimental/VirtualCamera/Source/VCamCore/Private/VCamComponent.cpp
--- a/Engine/Plugins/Experimental/VirtualCamera/Source/VCamCore/Private/VCamComponent.cpp
+++ b/Engine/Plugins/Experimental/VirtualCamera/Source/VCamCore/Private/VCamComponent.cpp
@@ -4,15 +4,16 @@
 
 #include "Modifier/VCamModifier.h"
 #include "Modifier/VCamModifierContext.h"
+#include "Util/LevelViewportUtils.h"
+#include "VCamCoreCustomVersion.h"
 
+#include "Algo/ForEach.h"
 #include "CineCameraComponent.h"
-#include "Engine/GameInstance.h"
 #include "Engine/GameEngine.h"
 #include "Engine/InputDelegateBinding.h"
 #include "EnhancedInputComponent.h"
 #include "EnhancedInputSubsystems.h"
 #include "Features/IModularFeatures.h"
-#include "Framework/Application/SlateApplication.h"
 #include "GameFramework/InputSettings.h"
 #include "ILiveLinkClient.h"
 #include "InputMappingContext.h"
@@ -24,7 +25,6 @@
 #include "Modules/ModuleManager.h"
 #include "Editor.h"
 #include "LevelEditor.h"
-#include "IAssetViewport.h"
 #include "SLevelViewport.h"
 
 #include "IConcertModule.h"
@@ -99,8 +99,8 @@ UVCamComponent::UVCamComponent()
 
 void UVCamComponent::OnComponentDestroyed(bool bDestroyingHierarchy)
 {
-	bLockViewportToCamera = false;
-	UpdateActorLock();
+	bLockViewportToCamera_DEPRECATED = false;
+	UpdateActorViewportLocks();
 
 	for (UVCamOutputProviderBase* Provider : OutputProviders)
 	{
@@ -300,6 +300,18 @@ void UVCamComponent::OnAttachmentChanged()
 #endif
 }
 
+void UVCamComponent::Serialize(FArchive& Ar)
+{
+	using namespace UE::VCamCore;
+	Super::Serialize(Ar);
+	Ar.UsingCustomVersion(FVCamCoreCustomVersion::GUID);
+	
+	if (Ar.IsLoading() && Ar.CustomVer(FVCamCoreCustomVersion::GUID) < FVCamCoreCustomVersion::MoveTargetViewportFromComponentToOutput)
+	{
+		Algo::ForEach(ViewportLocker.Locks, [this](TPair<EVCamTargetViewportID, FVCamViewportLockState>& Pair){ Pair.Value.bLockViewportToCamera = bLockViewportToCamera_DEPRECATED; });
+	}
+}
+
 void UVCamComponent::PostLoad()
 {
 	Super::PostLoad();
@@ -337,7 +349,6 @@ void UVCamComponent::PreEditChange(FEditPropertyChain& PropertyAboutToChange)
 	{
 		static FName NAME_OutputProviders = GET_MEMBER_NAME_CHECKED(UVCamComponent, OutputProviders);
 		static FName NAME_ModifierStack = GET_MEMBER_NAME_CHECKED(UVCamComponent, ModifierStack);
-		static FName NAME_GeneratedModifier = GET_MEMBER_NAME_CHECKED(FModifierStackEntry, GeneratedModifier);
 		static FName NAME_Enabled = GET_MEMBER_NAME_CHECKED(UVCamComponent, bEnabled);
 
 		const FName MemberPropertyName = MemberProperty->GetFName();
@@ -371,19 +382,12 @@ void UVCamComponent::PostEditChangeProperty(FPropertyChangedEvent& PropertyChang
 	FProperty* Property = PropertyChangedEvent.MemberProperty;
 	if (Property && PropertyChangedEvent.ChangeType != EPropertyChangeType::Interactive)
 	{
-		static FName NAME_LockViewportToCamera = GET_MEMBER_NAME_CHECKED(UVCamComponent, bLockViewportToCamera);
 		static FName NAME_Enabled = GET_MEMBER_NAME_CHECKED(UVCamComponent, bEnabled);
 		static FName NAME_ModifierStack = GET_MEMBER_NAME_CHECKED(UVCamComponent, ModifierStack);
-		static FName NAME_TargetViewport = GET_MEMBER_NAME_CHECKED(UVCamComponent, TargetViewport);
 		static FName NAME_InputProfile = GET_MEMBER_NAME_CHECKED(UVCamComponent, InputProfile);
-
 		const FName PropertyName = Property->GetFName();
 
-		if (PropertyName == NAME_LockViewportToCamera)
-		{
-			UpdateActorLock();
-		}
-		else if (PropertyName == NAME_Enabled)
+		if (PropertyName == NAME_Enabled)
 		{
 			// Only act here if we are a struct (like FModifierStackEntry)
 			if (!Property->GetOwner<UClass>())
@@ -395,20 +399,6 @@ void UVCamComponent::PostEditChangeProperty(FPropertyChangedEvent& PropertyChang
 		{
 			ValidateModifierStack();
 		}
-		else if (PropertyName == NAME_TargetViewport)
-		{
-			if (bEnabled)
-			{
-				SetEnabled(false);
-				SetEnabled(true);
-
-				if (bLockViewportToCamera)
-				{
-					SetActorLock(false);
-					SetActorLock(true);
-				}
-			}
-		}
 		else if (PropertyName == NAME_InputProfile)
 		{
 			ApplyInputProfile();
@@ -432,48 +422,20 @@ void UVCamComponent::PostEditChangeChainProperty(FPropertyChangedChainEvent& Pro
 	if (Property && PropertyChangedEvent.ChangeType != EPropertyChangeType::Interactive)
 	{
 		static FName NAME_OutputProviders = GET_MEMBER_NAME_CHECKED(UVCamComponent, OutputProviders);
+		static FName NAME_TargetViewport = UVCamOutputProviderBase::GetTargetViewportPropertyName();
 
 		if (Property->GetFName() == NAME_OutputProviders)
 		{
-			FProperty* ActualProperty = PropertyChangedEvent.PropertyChain.GetActiveNode()->GetNextNode() ? PropertyChangedEvent.PropertyChain.GetActiveNode()->GetNextNode()->GetValue() : nullptr;
+			FProperty* ActualProperty = PropertyChangedEvent.PropertyChain.GetActiveNode()->GetNextNode()
+				? PropertyChangedEvent.PropertyChain.GetActiveNode()->GetNextNode()->GetValue()
+				: nullptr;
 			if (ActualProperty == nullptr)
 			{
-				const int32 ChangedIndex = PropertyChangedEvent.GetArrayIndex(PropertyChangedEvent.GetPropertyName().ToString());
-				if (PropertyChangedEvent.ChangeType == EPropertyChangeType::ValueSet)
-				{
-					if (OutputProviders.IsValidIndex(ChangedIndex))
-					{
-						UVCamOutputProviderBase* ChangedProvider = OutputProviders[ChangedIndex];
-
-						// If we changed the output type, be sure to delete the old one before setting up the new one
-						if (SavedOutputProviders.IsValidIndex(ChangedIndex) && (SavedOutputProviders[ChangedIndex] != ChangedProvider))
-						{
-							DestroyOutputProvider(SavedOutputProviders[ChangedIndex]);
-						}
-
-						// We only Initialize a provider if they're able to be updated
-						// If they later become able to be updated then they will be
-						// Initialized inside the Update() loop
-						if (ChangedProvider && ShouldUpdateOutputProviders())
-						{
-							ChangedProvider->Initialize();
-						}
-					}
-				}
-				else if (PropertyChangedEvent.ChangeType == EPropertyChangeType::ArrayRemove)
-				{
-					if (SavedOutputProviders.IsValidIndex(ChangedIndex))
-					{
-						DestroyOutputProvider(SavedOutputProviders[ChangedIndex]);
-					}
-				}
-				else if (PropertyChangedEvent.ChangeType == EPropertyChangeType::ArrayClear)
-				{
-					for (UVCamOutputProviderBase* ClearedProvider : SavedOutputProviders)
-					{
-						DestroyOutputProvider(ClearedProvider);
-					}
-				}
+				OnOutputProvidersEdited(PropertyChangedEvent);
+			}
+			else if (ActualProperty->GetFName() == NAME_TargetViewport)
+			{
+				OnTargetViewportEdited();
 			}
 
 			// We created this in PreEditChange, so we need to always get rid of it
@@ -483,6 +445,57 @@ void UVCamComponent::PostEditChangeChainProperty(FPropertyChangedChainEvent& Pro
 
 	Super::PostEditChangeChainProperty(PropertyChangedEvent);
 }
+
+void UVCamComponent::OnOutputProvidersEdited(FPropertyChangedChainEvent& PropertyChangedEvent)
+{
+	const int32 ChangedIndex = PropertyChangedEvent.GetArrayIndex(PropertyChangedEvent.GetPropertyName().ToString());
+	if (PropertyChangedEvent.ChangeType == EPropertyChangeType::ValueSet)
+	{
+		if (OutputProviders.IsValidIndex(ChangedIndex))
+		{
+			UVCamOutputProviderBase* ChangedProvider = OutputProviders[ChangedIndex];
+
+			// If we changed the output type, be sure to delete the old one before setting up the new one
+			if (SavedOutputProviders.IsValidIndex(ChangedIndex) && (SavedOutputProviders[ChangedIndex] != ChangedProvider))
+			{
+				DestroyOutputProvider(SavedOutputProviders[ChangedIndex]);
+			}
+
+			// We only Initialize a provider if they're able to be updated
+			// If they later become able to be updated then they will be
+			// Initialized inside the Update() loop
+			if (ChangedProvider && ShouldUpdateOutputProviders())
+			{
+				ChangedProvider->Initialize();
+			}
+		}
+	}
+	else if (PropertyChangedEvent.ChangeType == EPropertyChangeType::ArrayRemove)
+	{
+		if (SavedOutputProviders.IsValidIndex(ChangedIndex))
+		{
+			DestroyOutputProvider(SavedOutputProviders[ChangedIndex]);
+		}
+	}
+	else if (PropertyChangedEvent.ChangeType == EPropertyChangeType::ArrayClear)
+	{
+		for (UVCamOutputProviderBase* ClearedProvider : SavedOutputProviders)
+		{
+			DestroyOutputProvider(ClearedProvider);
+		}
+	}
+}
+
+void UVCamComponent::OnTargetViewportEdited()
+{
+	if (bEnabled)
+	{
+		SetEnabled(false);
+		SetEnabled(true);
+		UpdateActorViewportLocks();
+	}
+}
+
 #endif // WITH_EDITOR
 
 void UVCamComponent::AddInputMappingContext(const UVCamModifier* Modifier)
@@ -689,11 +702,8 @@ void UVCamComponent::Update()
 		// Ensure the actor lock reflects the state of the lock property
 		// This is needed as UActorComponent::ConsolidatedPostEditChange will cause the component to be reconstructed on PostEditChange
 		// if the component is inherited
-		if (bLockViewportToCamera != bIsLockedToViewport)
-		{
-			UpdateActorLock();
-		}
-
+		UpdateActorViewportLocks();
+		
 		FLiveLinkCameraBlueprintData InitialLiveLinkData;
 		GetLiveLinkDataForCurrentFrame(InitialLiveLinkData);
 
@@ -1234,77 +1244,24 @@ float UVCamComponent::GetDeltaTime()
 	return DeltaTime;
 }
 
-void UVCamComponent::UpdateActorLock()
+void UVCamComponent::UpdateActorViewportLocks()
 {
-	if (GetTargetCamera() == nullptr)
+	if (const UCineCameraComponent* Camera = GetTargetCamera()
+		; Camera && ensure(Camera->GetOwner()))
 	{
-		UE_LOG(LogVCamComponent, Warning, TEXT("UpdateActorLock has been called, but there is no valid TargetCamera!"));
-		return;
+		UE::VCamCore::LevelViewportUtils::Private::UpdateViewportLocksFromOutputs(OutputProviders, ViewportLocker, Camera->GetOwner());
 	}
+}
 
-	for (const FWorldContext& Context : GEngine->GetWorldContexts())
+void UVCamComponent::UnlockAllViewports()
+{
+	if (const UCineCameraComponent* Camera = GetTargetCamera()
+		; Camera && ensure(Camera->GetOwner()))
 	{
-#if WITH_EDITOR
-		if (Context.WorldType == EWorldType::Editor)
-		{
-			if (FLevelEditorViewportClient* LevelViewportClient = GetTargetLevelViewportClient())
-			{
-				if (bLockViewportToCamera)
-				{
-					Backup_ActorLock = LevelViewportClient->GetActiveActorLock();
-					LevelViewportClient->SetActorLock(GetTargetCamera()->GetOwner());
-					// If bLockedCameraView is not true then the viewport is locked to the actor's transform and not the camera component
-					LevelViewportClient->bLockedCameraView = true;
-					bIsLockedToViewport = true;
-				}
-				else if (Backup_ActorLock.IsValid())
-				{
-					LevelViewportClient->SetActorLock(Backup_ActorLock.Get());
-					Backup_ActorLock = nullptr;
-					// If bLockedCameraView is not true then the viewport is locked to the actor's transform and not the camera component
-					LevelViewportClient->bLockedCameraView = true;
-					bIsLockedToViewport = false;
-				}
-				else
-				{
-					LevelViewportClient->SetActorLock(nullptr);
-					bIsLockedToViewport = false;
-				}
-			}
-		}
-		else
-#endif
-		{
-			UWorld* ActorWorld = Context.World();
-			if (ActorWorld && ActorWorld->GetGameInstance())
-			{
-				APlayerController* PlayerController = ActorWorld->GetGameInstance()->GetFirstLocalPlayerController(ActorWorld);
-				if (PlayerController)
-				{
-					if (bLockViewportToCamera)
-					{
-						Backup_ViewTarget = PlayerController->GetViewTarget();
-						PlayerController->SetViewTarget(GetTargetCamera()->GetOwner());
-						bIsLockedToViewport = true;
-					}
-					else if (Backup_ViewTarget.IsValid())
-					{
-						PlayerController->SetViewTarget(Backup_ViewTarget.Get());
-						Backup_ViewTarget = nullptr;
-						bIsLockedToViewport = false;
-					}
-					else
-					{
-						PlayerController->SetViewTarget(nullptr);
-						bIsLockedToViewport = false;
-					}
-				}
-			}
-		}
+		UE::VCamCore::LevelViewportUtils::Private::UnlockAllViewports(ViewportLocker);
 	}
 }
 
-
 void UVCamComponent::DestroyOutputProvider(UVCamOutputProviderBase* Provider)
 {
 	if (Provider)
@@ -1450,160 +1407,7 @@ void UVCamComponent::FindModifiedStackEntry(int32& ModifiedStackIndex, bool& bIs
 
 }
 
-TSharedPtr<FSceneViewport> UVCamComponent::GetTargetSceneViewport() const
-{
-	TSharedPtr<FSceneViewport> SceneViewport;
-
 #if WITH_EDITOR
-	if (GIsEditor)
-	{
-		for (const FWorldContext& Context : GEngine->GetWorldContexts())
-		{
-			if (Context.WorldType == EWorldType::PIE)
-			{
-				FSlatePlayInEditorInfo* SlatePlayInEditorSession = GEditor->SlatePlayInEditorMap.Find(Context.ContextHandle);
-				if (SlatePlayInEditorSession)
-				{
-					if (SlatePlayInEditorSession->DestinationSlateViewport.IsValid())
-					{
-						TSharedPtr<IAssetViewport> DestinationLevelViewport = SlatePlayInEditorSession->DestinationSlateViewport.Pin();
-						SceneViewport = DestinationLevelViewport->GetSharedActiveViewport();
-					}
-					else if (SlatePlayInEditorSession->SlatePlayInEditorWindowViewport.IsValid())
-					{
-						SceneViewport = SlatePlayInEditorSession->SlatePlayInEditorWindowViewport;
-					}
-
-					// If PIE is active always choose it
-					break;
-				}
-			}
-			else if (Context.WorldType == EWorldType::Editor)
-			{
-				if (FLevelEditorViewportClient* LevelViewportClient = GetTargetLevelViewportClient())
-				{
-					TSharedPtr<SEditorViewport> ViewportWidget = LevelViewportClient->GetEditorViewportWidget();
-					if (ViewportWidget.IsValid())
-					{
-						SceneViewport = ViewportWidget->GetSceneViewport();
-					}
-				}
-			}
-		}
-	}
-#else
-	if (UGameEngine* GameEngine = Cast<UGameEngine>(GEngine))
-	{
-		SceneViewport = GameEngine->SceneViewport;
-	}
-#endif
-
-	return SceneViewport;
-}
-
-TWeakPtr<SWindow> UVCamComponent::GetTargetInputWindow() const
-{
-	TWeakPtr<SWindow> InputWindow;
-
-#if WITH_EDITOR
-	if (GIsEditor)
-	{
-		for (const FWorldContext& Context : GEngine->GetWorldContexts())
-		{
-			if (Context.WorldType == EWorldType::PIE)
-			{
-				FSlatePlayInEditorInfo* SlatePlayInEditorSession = GEditor->SlatePlayInEditorMap.Find(Context.ContextHandle);
-				if (SlatePlayInEditorSession)
-				{
-					if (SlatePlayInEditorSession->DestinationSlateViewport.IsValid())
-					{
-						TSharedPtr<IAssetViewport> DestinationLevelViewport = SlatePlayInEditorSession->DestinationSlateViewport.Pin();
-						InputWindow = FSlateApplication::Get().FindWidgetWindow(DestinationLevelViewport->AsWidget());
-					}
-					else if (SlatePlayInEditorSession->SlatePlayInEditorWindowViewport.IsValid())
-					{
-						InputWindow = SlatePlayInEditorSession->SlatePlayInEditorWindow;
-					}
-
-					// If PIE is active always choose it
-					break;
-				}
-			}
-			else if (Context.WorldType == EWorldType::Editor)
-			{
-				if (FLevelEditorViewportClient* LevelViewportClient = GetTargetLevelViewportClient())
-				{
-					TSharedPtr<SEditorViewport> ViewportWidget = LevelViewportClient->GetEditorViewportWidget();
-					if (ViewportWidget.IsValid())
-					{
-						InputWindow = FSlateApplication::Get().FindWidgetWindow(ViewportWidget.ToSharedRef());
-					}
-				}
-			}
-		}
-	}
-#else
-	if (UGameEngine* GameEngine = Cast<UGameEngine>(GEngine))
-	{
-		InputWindow = GameEngine->GameViewportWindow;
-	}
-#endif
-
-	return InputWindow;
-}
-
-#if WITH_EDITOR
-FLevelEditorViewportClient* UVCamComponent::GetTargetLevelViewportClient() const
-{
-	FLevelEditorViewportClient* OutClient = nullptr;
-
-	TSharedPtr<SLevelViewport> LevelViewport = GetTargetLevelViewport();
-	if (LevelViewport.IsValid())
-	{
-		OutClient = &LevelViewport->GetLevelViewportClient();
-	}
-
-	return OutClient;
-}
-
-TSharedPtr<SLevelViewport> UVCamComponent::GetTargetLevelViewport() const
-{
-	TSharedPtr<SLevelViewport> OutLevelViewport = nullptr;
-
-	if (TargetViewport == EVCamTargetViewportID::CurrentlySelected)
-	{
-		if (FLevelEditorModule* LevelEditorModule = FModuleManager::GetModulePtr<FLevelEditorModule>(VCamComponent::LevelEditorName))
-		{
-			OutLevelViewport = LevelEditorModule->GetFirstActiveLevelViewport();
-		}
-	}
-	else
-	{
-		if (GEditor)
-		{
-			for (FLevelEditorViewportClient* Client : GEditor->GetLevelViewportClients())
-			{
-				// We only care about the fully rendered 3D viewport...seems like there should be a better way to check for this
-				if (!Client->IsOrtho())
-				{
-					TSharedPtr<SLevelViewport> LevelViewport = StaticCastSharedPtr<SLevelViewport>(Client->GetEditorViewportWidget());
-					if (LevelViewport.IsValid())
-					{
-						const FString WantedViewportString = FString::Printf(TEXT("Viewport %d.Viewport"), (int32)TargetViewport);
-						const FString ViewportConfigKey = LevelViewport->GetConfigKey().ToString();
-						if (ViewportConfigKey.Contains(*WantedViewportString, ESearchCase::CaseSensitive, ESearchDir::FromStart))
-						{
-							OutLevelViewport = LevelViewport;
-							break;
-						}
-					}
-				}
-			}
-		}
-	}
-
-	return OutLevelViewport;
-}
 
 void UVCamComponent::OnMapChanged(UWorld* World, EMapChangeType ChangeType)
 {
diff --git a/Engine/Plugins/Experimental/VirtualCamera/Source/VCamCore/Private/VCamCoreCustomVersion.cpp b/Engine/Plugins/Experimental/VirtualCamera/Source/VCamCore/Private/VCamCoreCustomVersion.cpp
new file mode 100644
--- /dev/null
+++ b/Engine/Plugins/Experimental/VirtualCamera/Source/VCamCore/Private/VCamCoreCustomVersion.cpp
@@ -0,0 +1,7 @@
+// Copyright Epic Games, Inc. All Rights Reserved.
+
+#include "VCamCoreCustomVersion.h"
+#include "Serialization/CustomVersion.h"
+
+const FGuid UE::VCamCore::FVCamCoreCustomVersion::GUID(0x31234DBD, 0x21CAF31F, 0xAAA7A38F, 0xB4A32531);
+FCustomVersionRegistration GRegisterSnapshotCustomVersion(UE::VCamCore::FVCamCoreCustomVersion::GUID, UE::VCamCore::FVCamCoreCustomVersion::LatestVersion, TEXT("LevelSnapshotVersion"));
diff --git a/Engine/Plugins/Experimental/VirtualCamera/Source/VCamCore/Private/VCamCoreCustomVersion.h b/Engine/Plugins/Experimental/VirtualCamera/Source/VCamCore/Private/VCamCoreCustomVersion.h
new file mode 100644
--- /dev/null
+++ b/Engine/Plugins/Experimental/VirtualCamera/Source/VCamCore/Private/VCamCoreCustomVersion.h
@@ -0,0 +1,34 @@
+// Copyright Epic Games, Inc. All Rights Reserved.
+
+#pragma once
+
+#include "CoreMinimal.h"
+#include "Misc/Guid.h"
+
+namespace UE::VCamCore
+{
+	struct FVCamCoreCustomVersion
+	{
+		enum Type
+		{
+			/** Before any version changes were made in the plugin */
+			BeforeCustomVersionWasAdded = 0,
+
+			/**
+			 * The following properties were changed (for 5.2):
+			 * - UVCamComponent::TargetViewport was moved to UVCamOutputProviderBase
+			 * - UVCamComponent::bLockViewportToCamera was reworked to work for multiple viewports in FVCamViewportLocker (now part of UVCamComponent).
+			 */
+			MoveTargetViewportFromComponentToOutput = 1,
+
+			// -----<new versions can be added above this line>-------------------------------------------------
+			VersionPlusOne,
+			LatestVersion = VersionPlusOne - 1
+		};
+
+		const static FGuid GUID;
+
+	private:
+		FVCamCoreCustomVersion() = delete;
+	};
+}
diff --git a/Engine/Plugins/Experimental/VirtualCamera/Source/VCamCore/Public/Output/VCamOutputProviderBase.h b/Engine/Plugins/Experimental/VirtualCamera/Source/VCamCore/Public/Output/VCamOutputProviderBase.h
--- a/Engine/Plugins/Experimental/VirtualCamera/Source/VCamCore/Public/Output/VCamOutputProviderBase.h
+++ b/Engine/Plugins/Experimental/VirtualCamera/Source/VCamCore/Public/Output/VCamOutputProviderBase.h
@@ -4,6 +4,7 @@
 
 #include "CoreMinimal.h"
 #include "CineCameraComponent.h"
+#include "EVCamTargetViewportID.h"
 #include "VPFullScreenUserWidget.h"
 #include "VCamOutputProviderBase.generated.h"
 
@@ -24,7 +25,17 @@ class VCAMCORE_API UVCamOutputProviderBase : public UObject
 	GENERATED_BODY()
 public:
 
+	/** Override the default output resolution with a custom value - NOTE you must toggle bIsActive off then back on for this to take effect */
+	UPROPERTY(EditAnywhere, BlueprintReadWrite, Category = "Output", meta = (DisplayPriority = "5"))
+	bool bUseOverrideResolution = false;
+
+	/** When bUseOverrideResolution is set, use this custom resolution */
+	UPROPERTY(EditAnywhere, BlueprintReadWrite, Category = "Output", meta = (DisplayPriority = "6"), meta = (EditCondition = "bUseOverrideResolution", ClampMin = 1))
+	FIntPoint OverrideResolution = { 2048, 1536 };
+
+	//~ Begin UObject Interface
 	virtual void BeginDestroy() override;
+	//~ End UObject Interface
 
 	// Called when the provider is brought online such as after instantiating or loading a component containing this provider 
 	// Use Initialize for any setup logic that needs to survive between Start / Stop cycles such as spawning transient objects 
@@ -46,11 +57,6 @@ public:
 	// Called to turn on or off this output provider
 	UFUNCTION(BlueprintCallable, Category = "Output")
 	void SetActive(const bool bInActive);
-
-#if WITH_EDITOR
-	virtual void PostEditChangeProperty(FPropertyChangedEvent& PropertyChangedEvent) override;
-#endif
-
 	// Returns if this output provider is currently active or not
 	UFUNCTION(BlueprintPure, Category = "Output")
 	bool IsActive() const { return bIsActive; };
@@ -59,25 +65,19 @@ public:
 	UFUNCTION(BlueprintPure, Category = "Output")
 	bool IsInitialized() const { return bInitialized; };
 
-	// Sets the TargetCamera parameter
 	UFUNCTION(BlueprintCallable, Category = "Output")
 	void SetTargetCamera(const UCineCameraComponent* InTargetCamera);
 
-	// Sets the UMG class to be rendered in this output provider
+	UFUNCTION(BlueprintPure, Category = "Output")
+	EVCamTargetViewportID GetTargetViewport() const { return TargetViewport; }
 	UFUNCTION(BlueprintCallable, Category = "Output")
-	void SetUMGClass(const TSubclassOf<UUserWidget> InUMGClass);
-
-	// The UMG class to be rendered in this output provider
-	UPROPERTY(EditAnywhere, BlueprintReadWrite, Category = "Output", meta = (DisplayName="UMG Overlay", DisplayPriority = "2"))
-	TSubclassOf<UUserWidget> UMGClass;
-
-	// Override the default output resolution with a custom value - NOTE you must toggle bIsActive off then back on for this to take effect
-	UPROPERTY(EditAnywhere, BlueprintReadWrite, Category = "Output", meta = (DisplayPriority = "3"))
-	bool bUseOverrideResolution = false;
-
-	// When bUseOverrideResolution is set, use this custom resolution
-	UPROPERTY(EditAnywhere, BlueprintReadWrite, Category = "Output", meta = (DisplayPriority = "4"), meta = (EditCondition = "bUseOverrideResolution", ClampMin = 1))
-	FIntPoint OverrideResolution = { 2048, 1536 };
+	void SetTargetViewport(EVCamTargetViewportID Value) { TargetViewport = Value; }
+	static FName GetTargetViewportPropertyName() { return GET_MEMBER_NAME_CHECKED(UVCamOutputProviderBase, TargetViewport); }
+	
+	UFUNCTION(BlueprintPure, Category = "Output")
+	TSubclassOf<UUserWidget> GetUMGClass() const { return UMGClass; }
+	UFUNCTION(BlueprintCallable, Category = "Output")
+	void SetUMGClass(const TSubclassOf<UUserWidget> InUMGClass) { UMGClass = InUMGClass; }
 
 	UVPFullScreenUserWidget* GetUMGWidget() { return UMGWidget; };
 
@@ -104,7 +104,16 @@ public:
 	/** Calls the VCamModifierInterface on the widget if it exists and also requests any child VCam Widgets to reconnect */
 	void NotifyWidgetOfComponentChange() const;
 
+	//~ Begin UObject Interface
+	virtual void Serialize(FArchive& Ar) override;
+	//~ End UObject Interface
+
+#if WITH_EDITOR
+	virtual void PostEditChangeProperty(FPropertyChangedEvent& PropertyChangedEvent) override;
+#endif
+
 protected:
+	
 	// If set, this output provider will execute every frame
 	UPROPERTY(EditAnywhere, BlueprintReadWrite, Category = "Output", meta = (DisplayPriority = "1"))
 	bool bIsActive = false;
@@ -112,7 +121,11 @@ protected:
 	UPROPERTY(Transient)
 	bool bInitialized = false;
 
+	UPROPERTY(Transient)
 	EVPWidgetDisplayType DisplayType = EVPWidgetDisplayType::PostProcess;
+	
+	UPROPERTY(Transient)
+	TObjectPtr<UVPFullScreenUserWidget> UMGWidget = nullptr;
 
 	virtual void CreateUMG();
 
@@ -129,12 +142,21 @@ protected:
 	TSharedPtr<SLevelViewport> GetTargetLevelViewport() const;
 #endif
 
-	UPROPERTY(Transient)
-	TObjectPtr<UVPFullScreenUserWidget> UMGWidget = nullptr;
-
 private:
-	bool IsOuterComponentEnabled() const;
 
+	/** Which viewport to use for this VCam */
+	UPROPERTY(EditAnywhere, BlueprintGetter = "GetTargetViewport", BlueprintSetter = "SetTargetViewport", Category = "Output", meta = (DisplayPriority = "2"))
+	EVCamTargetViewportID TargetViewport = EVCamTargetViewportID::CurrentlySelected;
+	
+	/** The UMG class to be rendered in this output provider */
+	UPROPERTY(EditAnywhere, BlueprintGetter = "GetUMGClass", BlueprintSetter = "SetUMGClass", Category = "Output", meta = (DisplayName="UMG Overlay", DisplayPriority = "3"))
+	TSubclassOf<UUserWidget> UMGClass;
+	
+	UPROPERTY(Transient)
 	TSoftObjectPtr<UCineCameraComponent> TargetCamera;
+	
+	UPROPERTY(Transient)
 	bool bWasActive = false;
+	
+	bool IsOuterComponentEnabled() const;
 };
diff --git a/Engine/Plugins/Experimental/VirtualCamera/Source/VCamCore/Public/Util/VCamViewportLocker.h b/Engine/Plugins/Experimental/VirtualCamera/Source/VCamCore/Public/Util/VCamViewportLocker.h
new file mode 100644
--- /dev/null
+++ b/Engine/Plugins/Experimental/VirtualCamera/Source/VCamCore/Public/Util/VCamViewportLocker.h
@@ -0,0 +1,46 @@
+// Copyright Epic Games, Inc. All Rights Reserved.
+
+#pragma once
+
+#include "CoreMinimal.h"
+#include "EVCamTargetViewportID.h"
+#include "VCamViewportLocker.generated.h"
+
+USTRUCT()
+struct FVCamViewportLockState
+{
+	GENERATED_BODY()
+	
+	/** Whether the user wants the viewport to be locked */
+	UPROPERTY(EditAnywhere, Category = "Viewport")
+	bool bLockViewportToCamera = false;
+	
+	/** Whether this viewport is currently locked */
+	UPROPERTY(Transient)
+	bool bIsLockedToViewport = false;
+	
+	/** Used for editor */
+	UPROPERTY(Transient)
+	TWeakObjectPtr<AActor> Backup_ActorLock;
+	
+	/** Used for gameplay */
+	UPROPERTY(Transient)
+	TWeakObjectPtr<AActor> Backup_ViewTarget;
+};
+
+/**
+ * Keeps track of which viewports are locked
+ */
+USTRUCT()
+struct FVCamViewportLocker
+{
+	GENERATED_BODY()
+
+	UPROPERTY(EditAnywhere, EditFixedSize, Category = "Viewport")
+	TMap<EVCamTargetViewportID, FVCamViewportLockState> Locks {
+		{ EVCamTargetViewportID::Viewport1, {} },
+		{ EVCamTargetViewportID::Viewport2, {} },
+		{ EVCamTargetViewportID::Viewport3, {} },
+		{ EVCamTargetViewportID::Viewport4, {} }
+	};
+};
diff --git a/Engine/Plugins/Experimental/VirtualCamera/Source/VCamCore/Public/VCamComponent.h b/Engine/Plugins/Experimental/VirtualCamera/Source/VCamCore/Public/VCamComponent.h
--- a/Engine/Plugins/Experimental/VirtualCamera/Source/VCamCore/Public/VCamComponent.h
+++ b/Engine/Plugins/Experimental/VirtualCamera/Source/VCamCore/Public/VCamComponent.h
@@ -13,8 +13,10 @@
 #include "GameplayTagContainer.h"
 #include "Roles/LiveLinkCameraTypes.h"
 #include "VCamInputSettings.h"
+#include "Util/VCamViewportLocker.h"
 
 #if WITH_EDITOR
+#include "Engine/GameEngine.h"
 #include "UnrealEdMisc.h"
 #include "VCamMultiUser.h"
 #endif
@@ -41,66 +43,73 @@ class VCAMCORE_API UVCamComponent : public USceneComponent
 	GENERATED_BODY()
 	friend class UVCamModifier;
 public:
+
+	/**
+	 * There are situations in the editor where the component may be replaced by another component as part of the actor being reconstructed
+	 * This event will notify you of that change and give you a reference to the new component.
+	 * Bindings will be copied to the new component so you do not need to rebind this event
+	 * 
+	 * Note: When the component is replaced you will need to get all properties on the component again such as Modifiers and Output Providers
+	 */
+	UPROPERTY(BlueprintAssignable, Category = "VirtualCamera")
+	FOnComponentReplaced OnComponentReplaced;
 	
 	UVCamComponent();
 
+	//~ Begin UActorComponent Interface
 	virtual void OnComponentDestroyed(bool bDestroyingHierarchy) override;
+	//~ End UActorComponent Interface
 
+	//~ Begin USceneComponent Interface
 	virtual void OnAttachmentChanged() override;
-
+	//~ End USceneComponent Interface
+	
+	//~ Begin UObject Interface
+	virtual void Serialize(FArchive& Ar) override;
 	virtual void PostLoad() override;
-
-	TSharedPtr<FSceneViewport> GetTargetSceneViewport() const;
-	TWeakPtr<SWindow> GetTargetInputWindow() const;
+	//~ End UObject Interface
 
 #if WITH_EDITOR
-	FLevelEditorViewportClient* GetTargetLevelViewportClient() const;
-	TSharedPtr<SLevelViewport> GetTargetLevelViewport() const;
-
 	virtual void CheckForErrors() override;
 	virtual void PreEditChange(FProperty* PropertyThatWillChange) override;
 	virtual void PreEditChange(FEditPropertyChain& PropertyAboutToChange) override;
 	virtual void PostEditChangeProperty(FPropertyChangedEvent& PropertyChangedEvent) override;
 	virtual void PostEditChangeChainProperty(FPropertyChangedChainEvent& PropertyChangedEvent) override;
+	void OnOutputProvidersEdited(FPropertyChangedChainEvent& PropertyChangedEvent);
+	void OnTargetViewportEdited();
 #endif // WITH_EDITOR
 
-	// There are situations in the editor where the component may be replaced by another component as part of the actor being reconstructed
-	// This event will notify you of that change and give you a reference to the new component. 
-	// Bindings will be copied to the new component so you do not need to rebind this event
-	// 
-	// Note: When the component is replaced you will need to get all properties on the component again such as Modifiers and Output Providers
-	UPROPERTY(BlueprintAssignable, Category = "VirtualCamera")
-	FOnComponentReplaced OnComponentReplaced;
-
 	UFUNCTION()
 	void HandleObjectReplaced(const TMap<UObject*, UObject*>& ReplacementMap);
 
 	bool CanUpdate() const;
-	
 	void Update();
 
-	// Adds the Input Mapping Context from a modifier, if it exists, to the input system 
+	/******************** Input ********************/
+	
+	/** Adds the Input Mapping Context from a modifier, if it exists, to the input system */
 	void AddInputMappingContext(const UVCamModifier* Modifier);
-
-	// Removes the Input Mapping Context from a modifier, if it exists, from the input system
+	/** Removes the Input Mapping Context from a modifier, if it exists, from the input system */
 	void RemoveInputMappingContext(const UVCamModifier* Modifier);
-
-	// Adds an explicitly provided Input Mapping Context to the input system
+	/** Adds an explicitly provided Input Mapping Context to the input system */
 	void AddInputMappingContext(UInputMappingContext* Context, int32 Priority);
-
-	// Removes an explicitly provided Input Mapping Context to the input system
+	/** Removes an explicitly provided Input Mapping Context to the input system */
 	void RemoveInputMappingContext(UInputMappingContext* Context);
-
-	// Attempts to apply key mapping settings from an input profile defined in VCam Input Settings
-	// Returns whether the profile was found and able to be applied
+	
+	/**
+	 * Attempts to apply key mapping settings from an input profile defined in VCam Input Settings
+	 * Returns whether the profile was found and able to be applied
+	 */
 	UFUNCTION(BlueprintCallable, Category = "VCam Input")
 	bool SetInputProfileFromName(const FName ProfileName);
 
-	// Tries to add a new Input Profile to the VCam Input Settings and populates it with any currently active player mappable keys
-	// Note: The set of currently active player mappable keys may be larger than the set of mappings in this component's Input Profile
-	//
-	// Optionally this function can update an existing profile, this will only add any mappable keys that don't currently exist in the profile but will not remove any existing mappings
-	// Returns whether the profile was successfully added or updated
+	/**
+	 * Tries to add a new Input Profile to the VCam Input Settings and populates it with any currently active player mappable keys
+	 * Note: The set of currently active player mappable keys may be larger than the set of mappings in this component's Input Profile
+	 *
+	 * Optionally this function can update an existing profile, this will only add any mappable keys that don't currently exist in the profile but will not remove any existing mappings
+	 * Returns whether the profile was successfully added or updated
+	 */
 	UFUNCTION(BlueprintCallable, Category="VCam Input", meta=(ReturnDisplayName = "Success"))
 	bool AddInputProfileWithCurrentlyActiveMappings(const FName ProfileName, bool bUpdateIfProfileAlreadyExists = true);
 
@@ -117,160 +126,231 @@ public:
 	UFUNCTION(BlueprintCallable, Category="VCam Input", meta=(ReturnDisplayName = "Key"))
 	FKey GetPlayerMappedKey(const FName MappingName) const;
 
-	// Sets if the VCamComponent will update every frame or not
+
+	
+	/** Sets whether the VCamComponent will update every frame */
 	UFUNCTION(BlueprintSetter)
 	void SetEnabled(bool bNewEnabled);
-
-	// Returns whether or not the VCamComponent will update every frame
+	/** @return Whether the VCamComponent will update every frame */
 	UFUNCTION(BlueprintGetter)
 	bool IsEnabled() const { return bEnabled; };
 
-	// Returns the Target CineCameraComponent
 	UFUNCTION(BlueprintPure, Category = "VirtualCamera")
 	UCineCameraComponent* GetTargetCamera() const;
 
-	// Add a modifier to the stack with a given name.
-	// If that name is already in use then the modifier will not be added.
-	// Returns the created modifier if the Add succeeded
+	
+	/******************** Modifier ********************/
+	
+	/**
+	 * Add a modifier to the stack with a given name.
+	 * If that name is already in use then the modifier will not be added.
+	 * Returns the created modifier if the Add succeeded
+	 */
 	UFUNCTION(BlueprintCallable, Category = "VirtualCamera", Meta = (DeterminesOutputType = "ModifierClass", DynamicOutputParam = "CreatedModifier", ReturnDisplayName = "Success"))
 	bool AddModifier(const FName Name, UPARAM(meta = (AllowAbstract = "false")) TSubclassOf<UVCamModifier> ModifierClass, UVCamModifier*& CreatedModifier);
-
-	// Insert a modifier to the stack with a given name and index.
-	// If that name is already in use then the modifier will not be added.
-	// The index must be between zero and the number of existing modifiers inclusive
-	// Returns the created modifier if the Add succeeded
+	/**
+	 * Insert a modifier to the stack with a given name and index.
+	 * If that name is already in use then the modifier will not be added.
+	 * The index must be between zero and the number of existing modifiers inclusive
+	 * Returns the created modifier if the Add succeeded.
+	 */
 	UFUNCTION(BlueprintCallable, Category = "VirtualCamera", Meta = (DeterminesOutputType = "ModifierClass", DynamicOutputParam = "CreatedModifier", ReturnDisplayName = "Success"))
 	bool InsertModifier(const FName Name, int32 Index, UPARAM(meta = (AllowAbstract = "false")) TSubclassOf<UVCamModifier> ModifierClass, UVCamModifier*& CreatedModifier);
-
 	// Moves an existing modifier in the stack from its current index to a new index
 	UFUNCTION(BlueprintCallable, Category = "VirtualCamera", Meta = (ReturnDisplayName = "Success"))
 	bool SetModifierIndex(int32 OriginalIndex, int32 NewIndex);
 
-	// Remove all Modifiers from the Stack.
+	/** Remove all Modifiers from the Stack. */
 	UFUNCTION(BlueprintCallable, Category = "VirtualCamera")
 	void RemoveAllModifiers();
-
-	// Remove the given Modifier from the Stack.
-	// Returns true if the modifier was removed successfully
+	/** @return Whether the modifier was removed successfully */
 	UFUNCTION(BlueprintCallable, Category = "VirtualCamera")
     bool RemoveModifier(const UVCamModifier* Modifier);
-
-	// Remove the Modifier at a specified index from the Stack.
-	// Returns true if the modifier was removed successfully
+	/** @return Whether the modifier was removed successfully */
 	UFUNCTION(BlueprintCallable, Category = "VirtualCamera")
 	bool RemoveModifierByIndex(const int ModifierIndex);
-
-	// Remove the Modifier with a specific name from the Stack.
-	// Returns true if the modifier was removed successfully
+	/** @return Whether the modifier was removed successfully. */
 	UFUNCTION(BlueprintCallable, Category = "VirtualCamera")
     bool RemoveModifierByName(const FName Name);
 
-	// Returns the number of Modifiers in the Component's Stack
 	UFUNCTION(BlueprintPure, Category = "VirtualCamera")
 	int32 GetNumberOfModifiers() const;
 
-	// Returns all the Modifiers in the Component's Stack
-	// Note: It's possible not all Modifiers will be valid (such as if the user has not set a class for the modifier in the details panel)
+	/**
+	 * Returns all the Modifiers in the Component's Stack
+	 * Note: It's possible not all Modifiers will be valid (such as if the user has not set a class for the modifier in the details panel)
+	 */
 	UFUNCTION(BlueprintPure, Category = "VirtualCamera")
 	void GetAllModifiers(TArray<UVCamModifier*>& Modifiers) const;
-
 	/** Returns all the modifier names used to identifying connection points. */
 	UFUNCTION(BlueprintPure, Category = "VirtualCamera")
 	TArray<FName> GetAllModifierNames() const;
-
-	// Returns the Modifier in the Stack with the given index if it exist.
+	/** Returns the Modifier in the Stack with the given index if it exist. */
 	UFUNCTION(BlueprintPure, Category = "VirtualCamera")
 	UVCamModifier* GetModifierByIndex(const int32 Index) const;
-
-	// Tries to find a Modifier in the Stack with the given name.
-	// The returned Modifier must be checked before it is used.
+	/**
+	 * Tries to find a Modifier in the Stack with the given name. The returned Modifier must be checked before it is used. */
 	UFUNCTION(BlueprintPure, Category = "VirtualCamera")
 	UVCamModifier* GetModifierByName(const FName Name) const;
 
-	// Given a specific Modifier class, returns a list of matching Modifiers
+	/** Given a specific Modifier class, returns a list of matching Modifiers */
 	UFUNCTION(BlueprintPure, Category = "VirtualCamera", Meta = (DeterminesOutputType = "ModifierClass", DynamicOutputParam = "FoundModifiers"))
 	void GetModifiersByClass(UPARAM(meta = (AllowAbstract = "false")) TSubclassOf<UVCamModifier> ModifierClass, TArray<UVCamModifier*>& FoundModifiers) const;
-
-	// Given a specific Interface class, returns a list of matching Modifiers
+	/** Given a specific Interface class, returns a list of matching Modifiers */
 	UFUNCTION(BlueprintPure, Category = "VirtualCamera")
 	void GetModifiersByInterface(UPARAM(meta = (AllowAbstract = "false")) TSubclassOf<UInterface> InterfaceClass, TArray<UVCamModifier*>& FoundModifiers) const;
 
-	/*
-	Sets the Modifier Context to a new instance of the provided class
-	@param ContextClass The Class to create the context from
-	@param CreatedContext The created Context, can be invalid if Context Class was None
-	*/
+	/**
+	 * Sets the Modifier Context to a new instance of the provided class
+	 * @param ContextClass The Class to create the context from
+	 * @param CreatedContext The created Context, can be invalid if Context Class was None
+	 */
 	UFUNCTION(BlueprintCallable, Category = "VirtualCamera", Meta = (DeterminesOutputType = "ContextClass", DynamicOutputParam = "CreatedContext", AllowAbstract = "false"))
 	void SetModifierContextClass(UPARAM(meta = (AllowAbstract = "false")) TSubclassOf<UVCamModifierContext> ContextClass, UVCamModifierContext*& CreatedContext);
 	
-	/*
-	Get the current Modifier Context
-	@return Current Context
-	*/
+	/**
+	 * Get the current Modifier Context
+	 * @return Current Context
+	 */
 	UFUNCTION(BlueprintPure, Category = "VirtualCamera")
 	UVCamModifierContext* GetModifierContext() const;
-
-	// Output Provider access
+	
+	
+	/******************** Output Provider access ********************/
 
 	UFUNCTION(BlueprintCallable, Category = "VirtualCamera", Meta = (DeterminesOutputType = "ProviderClass", DynamicOutputParam = "CreatedProvider", ReturnDisplayName = "Success"))
 	bool AddOutputProvider(UPARAM(meta = (AllowAbstract = "false")) TSubclassOf<UVCamOutputProviderBase> ProviderClass, UVCamOutputProviderBase*& CreatedProvider);
-
 	UFUNCTION(BlueprintCallable, Category = "VirtualCamera", Meta = (DeterminesOutputType = "ProviderClass", DynamicOutputParam = "CreatedProvider", ReturnDisplayName = "Success"))
 	bool InsertOutputProvider(int32 Index, UPARAM(meta = (AllowAbstract = "false")) TSubclassOf<UVCamOutputProviderBase> ProviderClass, UVCamOutputProviderBase*& CreatedProvider);
-
 	// Moves an existing Output Provider in the stack from its current index to a new index
 	UFUNCTION(BlueprintCallable, Category = "VirtualCamera", Meta = (ReturnDisplayName = "Success"))
 	bool SetOutputProviderIndex(int32 OriginalIndex, int32 NewIndex);
-
-	// Remove all Output Providers from the Component.
+	
 	UFUNCTION(BlueprintCallable, Category = "VirtualCamera")
 	void RemoveAllOutputProviders();
-
 	UFUNCTION(BlueprintCallable, Category = "VirtualCamera")
 	bool RemoveOutputProvider(const UVCamOutputProviderBase* Provider);
-
 	UFUNCTION(BlueprintCallable, Category = "VirtualCamera")
 	bool RemoveOutputProviderByIndex(const int32 ProviderIndex);
-
+	
 	UFUNCTION(BlueprintPure, Category = "VirtualCamera")
 	int32 GetNumberOfOutputProviders() const;
-
 	UFUNCTION(BlueprintPure, Category = "VirtualCamera")
 	void GetAllOutputProviders(TArray<UVCamOutputProviderBase*>& Providers) const;
 
 	UFUNCTION(BlueprintPure, Category = "VirtualCamera")
 	UVCamOutputProviderBase* GetOutputProviderByIndex(const int32 ProviderIndex) const;
-
 	UFUNCTION(BlueprintPure, Category = "VirtualCamera", Meta = (DeterminesOutputType = "ProviderClass", DynamicOutputParam = "FoundProviders"))
 	void GetOutputProvidersByClass(UPARAM(meta = (AllowAbstract = "false")) TSubclassOf<UVCamOutputProviderBase> ProviderClass, TArray<UVCamOutputProviderBase*>& FoundProviders) const;
 
+	/** Call this after modifying the InputProfile in code to update the player mapped keys */
+	void ApplyInputProfile();
+
+	UFUNCTION(BlueprintCallable, Category="VirtualCamera")
+	void GetLiveLinkDataForCurrentFrame(FLiveLinkCameraBlueprintData& LiveLinkData);
+
+	/**
+	 * Registers the given object with the VCamComponent's Input Component
+	 * This allows dynamic input bindings such as input events in blueprints to work correctly
+	 * Note: Ensure you call UnregisterObjectForInput when you are finished with the object
+	 * otherwise input events will still fire until GC actually destroys the object
+	 *
+	 * @param Object The object to register
+	 */
+	UFUNCTION(BlueprintCallable, Category="VirtualCamera")
+	void RegisterObjectForInput(UObject* Object);
+
+	/**
+	 * Unregisters the given object with the VCamComponent's Input Component
+	 *
+	 * @param Object The object to unregister
+	 */
+	UFUNCTION(BlueprintCallable, Category="VirtualCamera")
+	void UnregisterObjectForInput(UObject* Object) const;
+
+	/**
+	 * Returns a list of all player mappable keys that have been registered
+	 */
+	UFUNCTION(BlueprintCallable, Category="VirtualCamera")
+	TArray<FEnhancedActionKeyMapping> GetPlayerMappableKeys() const;
+
+	/**
+	 * Injects an input action. 
+	 */
+	UFUNCTION(BlueprintCallable, Category = "VirtualCamera", meta = (AutoCreateRefTerm = "Modifiers,Triggers"))
+	virtual void InjectInputForAction(const UInputAction* Action, FInputActionValue RawValue, const TArray<UInputModifier*>& Modifiers, const TArray<UInputTrigger*>& Triggers);
+	
+	/**
+	 * Injects an input vector for action.
+	 */
+	UFUNCTION(BlueprintCallable, Category = "VirtualCamera", meta = (AutoCreateRefTerm = "Modifiers,Triggers"))
+	virtual void InjectInputVectorForAction(const UInputAction* Action, FVector Value, const TArray<UInputModifier*>& Modifiers, const TArray<UInputTrigger*>& Triggers);
+
+	
+	/******************** Blueprint Getters & Setters ********************/
+
+
+	UFUNCTION(BlueprintPure, Category = "VirtualCamera")
+	FGameplayTag GetRole() const { return Role; }
+	UFUNCTION(BlueprintCallable, Category = "VirtualCamera")
+	void SetRole(FGameplayTag Value) { Role = Value; }
+	
+	UFUNCTION(BlueprintPure, Category = "VirtualCamera")
+	FLiveLinkSubjectName GetLiveLinkSubobject() const { return LiveLinkSubject; }
+	UFUNCTION(BlueprintCallable, Category = "VirtualCamera")
+	void SetLiveLinkSubobject(FLiveLinkSubjectName Value) { LiveLinkSubject = Value; }
+	
+	UFUNCTION(BlueprintPure, Category = "VirtualCamera")
+	bool GetDisableComponentWhenSpawnedBySequencer() const { return bDisableComponentWhenSpawnedBySequencer; }
+	UFUNCTION(BlueprintCallable, Category = "VirtualCamera")
+	void SetDisableComponentWhenSpawnedBySequencer(bool bValue) { bDisableComponentWhenSpawnedBySequencer = bValue; }
+	
+	UFUNCTION(BlueprintPure, Category = "VirtualCamera")
+	bool GetDisableOutputOnMultiUserReceiver() const { return bDisableOutputOnMultiUserReceiver; }
+	UFUNCTION(BlueprintCallable, Category = "VirtualCamera")
+	void SetDisableOutputOnMultiUserReceiver(bool bValue) { bDisableOutputOnMultiUserReceiver = bValue; }
+	
+	UFUNCTION(BlueprintPure, Category = "VCam Input")
+	const FVCamInputProfile& GetInputProfile() const { return InputProfile; }
+	UFUNCTION(BlueprintCallable, BlueprintInternalUseOnly, Category="VirtualCamera")
+	void SetInputProfile(const FVCamInputProfile& NewInputProfile);
+	
 private:
-	// Enabled state of the component
-	UPROPERTY(EditAnywhere, BlueprintSetter = SetEnabled, BlueprintGetter = IsEnabled, Category = "VirtualCamera")
+	
+	/** Whether the VCamComponent will update every frame */
+	UPROPERTY(EditAnywhere, BlueprintSetter = SetEnabled, BlueprintGetter = IsEnabled, Category = "VirtualCamera", meta = (DisplayPriority = "1"))
 	bool bEnabled = true;
-
-public:
+	
 	/**
 	 * The role of this virtual camera.  If this value is set and the corresponding tag set on the editor matches this value, then this
 	 * camera is the sender and the authority in the case when connected to a multi-user session.
 	 */
-	UPROPERTY(BlueprintReadWrite, EditAnywhere, Category = "VirtualCamera")
+	UPROPERTY(EditAnywhere, BlueprintGetter = "GetRole", BlueprintSetter = "SetRole", Category = "VirtualCamera", meta = (DisplayPriority = "2"))
 	FGameplayTag Role;
 
-	// LiveLink subject name for the incoming camera transform
-	UPROPERTY(EditAnywhere, BlueprintReadWrite, Category="VirtualCamera")
+	/** LiveLink subject name for the incoming camera transform */
+	UPROPERTY(EditAnywhere, BlueprintGetter = "GetLiveLinkSubobject", BlueprintSetter = "SetLiveLinkSubobject", Category="VirtualCamera", meta = (DisplayPriority = "3"))
 	FLiveLinkSubjectName LiveLinkSubject;
+	
+	/**
+	 * If true, render the viewport from the point of view of the parented CineCamera
+	 * 
+	 * This was moved to UVCamOutputProvider::TargetViewport.
+	 * See FVCamCoreCustomVersion::MoveTargetViewportFromComponentToOutput.
+	 */
+	UPROPERTY()
+	bool bLockViewportToCamera_DEPRECATED = false;
+	
+	/** Sync with output providers keeping track of which viewports are locked. */
+	UPROPERTY(EditAnywhere, Category = "VirtualCamera", meta = (DisplayPriority = "4"))
+	FVCamViewportLocker ViewportLocker;
 
-	// If true, render the viewport from the point of view of the parented CineCamera
-	UPROPERTY(EditAnywhere, BlueprintReadWrite, Category="VirtualCamera")
-	bool bLockViewportToCamera = false;
-
-	// If true, the component will force bEnabled to false when it is part of a spawnable in Sequencer
-	UPROPERTY(EditAnywhere, BlueprintReadWrite, Category = "VirtualCamera")
+	/** If true, the component will force bEnabled to false when it is part of a spawnable in Sequencer */
+	UPROPERTY(EditAnywhere, BlueprintGetter = "GetDisableComponentWhenSpawnedBySequencer", BlueprintSetter = "SetDisableComponentWhenSpawnedBySequencer", Category = "VirtualCamera", meta = (AllowPrivateAccess = "true", DisplayPriority = "5"))
 	bool bDisableComponentWhenSpawnedBySequencer = true;
 
 	/** Do we disable the output if the virtual camera is in a Multi-user session and the camera is a "receiver" from multi-user */
-	UPROPERTY(EditAnywhere, BlueprintReadWrite, AdvancedDisplay, Category = "VirtualCamera")
+	UPROPERTY(EditAnywhere, BlueprintGetter = "GetDisableOutputOnMultiUserReceiver", BlueprintSetter = "SetDisableOutputOnMultiUserReceiver", AdvancedDisplay, Category = "VirtualCamera", meta = (DisplayPriority = "6"))
 	bool bDisableOutputOnMultiUserReceiver = true;
 
 	/**
@@ -281,67 +361,63 @@ public:
 	UPROPERTY(EditAnywhere, AdvancedDisplay, Category="VirtualCamera", meta=(ForceUnits=ms, ClampMin = "11.0"), DisplayName="Update Frequencey")
 	float UpdateFrequencyMs = 66.6f;
 
-	/** Which viewport to use for this VCam */
-	UPROPERTY(EditAnywhere, BlueprintReadWrite, Category = "VirtualCamera")
-	EVCamTargetViewportID TargetViewport = EVCamTargetViewportID::CurrentlySelected;
+public:
+	
+	/**
+	 * Which viewport to use for this VCam.
+	 * 
+	 * This was moved to UVCamOutputProvider::TargetViewport.
+	 * See FVCamCoreCustomVersion::MoveTargetViewportFromComponentToOutput.
+	 */
+	UPROPERTY()
+	EVCamTargetViewportID TargetViewport_DEPRECATED = EVCamTargetViewportID::CurrentlySelected;
 
-	UPROPERTY(EditAnywhere, BlueprintReadWrite, BlueprintSetter=SetInputProfile, Category = "VirtualCamera")
+private:
+	
+	UPROPERTY(EditAnywhere, BlueprintGetter = "GetInputProfile", BlueprintSetter = "SetInputProfile", Category = "VirtualCamera")
 	FVCamInputProfile InputProfile;
 
 	/** List of Output Providers (executed in order) */
-	UPROPERTY(EditAnywhere, Instanced, Category="VirtualCamera")
+	UPROPERTY(EditAnywhere, Instanced, Category = "VirtualCamera")
 	TArray<TObjectPtr<UVCamOutputProviderBase>> OutputProviders;
-
-	/** Call this after modifying the InputProfile in code to update the player mapped keys */
-	void ApplyInputProfile();
 	
-	UFUNCTION(BlueprintCallable, BlueprintInternalUseOnly, Category="VirtualCamera")
-	void SetInputProfile(const FVCamInputProfile& NewInputProfile);
-
-	UFUNCTION(BlueprintCallable, Category="VirtualCamera")
-	void GetLiveLinkDataForCurrentFrame(FLiveLinkCameraBlueprintData& LiveLinkData);
+	/** Modifier Context object that can be accessed by the Modifier Stack */
+	UPROPERTY(EditAnywhere, Instanced, Category = "VirtualCamera")
+	TObjectPtr<UVCamModifierContext> ModifierContext;
 
-	/* Registers the given object with the VCamComponent's Input Component
-	 * This allows dynamic input bindings such as input events in blueprints to work correctly
-	 * Note: Ensure you call UnregisterObjectForInput when you are finished with the object
-	 * otherwise input events will still fire until GC actually destroys the object
-	 *
-	 * @param Object The object to register
+	/** List of Modifiers (executed in order) */
+	UPROPERTY(EditAnywhere, Category = "VirtualCamera")
+	TArray<FModifierStackEntry> ModifierStack;
+	
+	/**
+	 * From Ben H: Mark this as Transient/DuplicateTransient so that it is saved on the BP CDO and nowhere else and 
+	 * handled correctly during duplication operations (copy/paste etc)
 	 */
-	UFUNCTION(BlueprintCallable, Category="VirtualCamera")
-	void RegisterObjectForInput(UObject* Object);
+	UPROPERTY(Transient, DuplicateTransient)
+	TObjectPtr<UInputComponent> InputComponent;
 
-	/* Unregisters the given object with the VCamComponent's Input Component
-	 *
-	 * @param Object The object to unregister
-	 */
-	UFUNCTION(BlueprintCallable, Category="VirtualCamera")
-	void UnregisterObjectForInput(UObject* Object) const;
+	/** Store the Input Mapping Contexts that have been added via this component */
+	UPROPERTY(Transient, DuplicateTransient)
+	TArray<TObjectPtr<const UInputMappingContext>> AppliedInputContexts;
+	
+	double LastEvaluationTime;
 
-	/**
-	 * Returns a list of all player mappable keys that have been registered
-	 */
-	UFUNCTION(BlueprintCallable, Category="VirtualCamera")
-	TArray<FEnhancedActionKeyMapping> GetPlayerMappableKeys() const;
+	TWeakObjectPtr<AActor> Backup_ActorLock;
+	TWeakObjectPtr<AActor> Backup_ViewTarget;
 
-	/**
-	* Injects an input action. 
-	*/
-	UFUNCTION(BlueprintCallable, Category = "VirtualCamera", meta = (AutoCreateRefTerm = "Modifiers,Triggers"))
-	virtual void InjectInputForAction(const UInputAction* Action, FInputActionValue RawValue, const TArray<UInputModifier*>& Modifiers, const TArray<UInputTrigger*>& Triggers);
+	TArray<UVCamOutputProviderBase*> SavedOutputProviders;
+	TArray<FModifierStackEntry> SavedModifierStack;
 	
-	/**
-	* Injects an input vector for action.
-	*/
-	UFUNCTION(BlueprintCallable, Category = "VirtualCamera", meta = (AutoCreateRefTerm = "Modifiers,Triggers"))
-	virtual void InjectInputVectorForAction(const UInputAction* Action, FVector Value, const TArray<UInputModifier*>& Modifiers, const TArray<UInputTrigger*>& Triggers);
+	/** Variable used for pausing update on editor objects while PIE is running */
+	bool bIsEditorObjectButPIEIsRunning = false;
 
-private:	
+	bool bIsInputRegistered = false;
+	
 	static void CopyLiveLinkDataToCamera(const FLiveLinkCameraBlueprintData& LiveLinkData, UCineCameraComponent* CameraComponent);
 
 	float GetDeltaTime();
-	void SetActorLock(bool bNewActorLock) { bLockViewportToCamera = bNewActorLock; UpdateActorLock(); }
-	void UpdateActorLock();
+	void UpdateActorViewportLocks();
+	void UnlockAllViewports();
 	void DestroyOutputProvider(UVCamOutputProviderBase* Provider);
 	void ResetAllOutputProviders();
 
@@ -405,40 +481,7 @@ private:
 	 */
 	IEnhancedInputSubsystemInterface* GetEnhancedInputSubsystemInterface() const;
 
-	double LastEvaluationTime;
-
-	TWeakObjectPtr<AActor> Backup_ActorLock;
-	TWeakObjectPtr<AActor> Backup_ViewTarget;
-
-	TArray<UVCamOutputProviderBase*> SavedOutputProviders;
-	TArray<FModifierStackEntry> SavedModifierStack;
-
-	// Modifier Context object that can be accessed by the Modifier Stack
-	UPROPERTY(EditAnywhere, Instanced, Category = "VirtualCamera")
-	TObjectPtr<UVCamModifierContext> ModifierContext;
-
-	// List of Modifiers (executed in order)
-	UPROPERTY(EditAnywhere, Category = "VirtualCamera")
-	TArray<FModifierStackEntry> ModifierStack;
-
-	// Variable used for pausing update on editor objects while PIE is running
-	bool bIsEditorObjectButPIEIsRunning = false;
-
-	UPROPERTY(Transient)
-	bool bIsLockedToViewport = false;
-
-	// From Ben H: Mark this as Transient/DuplicateTransient so that it is saved on the BP CDO and nowhere else and
-	// handled correctly during duplication operations (copy/paste etc)
-	UPROPERTY(Transient, DuplicateTransient)
-	TObjectPtr<UInputComponent> InputComponent;
-
-	// Utility functions for registering and unregistering our input component with the correct input system
+	/** Utility functions for registering and unregistering our input component with the correct input system */
 	void RegisterInputComponent();
 	void UnregisterInputComponent();
-
-	// Store the Input Mapping Contexts that have been added via this component
-	UPROPERTY(Transient, DuplicateTransient)
-	TArray<TObjectPtr<const UInputMappingContext>> AppliedInputContexts;
-
-	bool bIsInputRegistered = false;
 };
