VCamBlueprintNodes should not have any PlatformAllowList nor TargetDenyList properties

diff --git a/Engine/Plugins/Experimental/VirtualCamera/VirtualCamera.uplugin b/Engine/Plugins/Experimental/VirtualCamera/VirtualCamera.uplugin
--- a/Engine/Plugins/Experimental/VirtualCamera/VirtualCamera.uplugin
+++ b/Engine/Plugins/Experimental/VirtualCamera/VirtualCamera.uplugin
@@ -33,9 +33,7 @@
 		{
 			"Name": "VCamBlueprintNodes",
 			"Type": "UncookedOnly",
-			"LoadingPhase": "Default",
-			"PlatformAllowList": [ "Win64", "Linux" ],
-			"TargetDenyList": [ "Server" ]
+			"LoadingPhase": "Default"
 		},
 		{
 			"Name": "VCamCore",
