Fix non-unity error

diff --git a/Engine/Plugins/Experimental/VirtualCamera/Source/VCamCore/Private/Util/LevelViewportUtils.cpp b/Engine/Plugins/Experimental/VirtualCamera/Source/VCamCore/Private/Util/LevelViewportUtils.cpp
--- a/Engine/Plugins/Experimental/VirtualCamera/Source/VCamCore/Private/Util/LevelViewportUtils.cpp
+++ b/Engine/Plugins/Experimental/VirtualCamera/Source/VCamCore/Private/Util/LevelViewportUtils.cpp
@@ -7,7 +7,8 @@
 #include "Output/VCamOutputProviderBase.h"
 #include "Util/VCamViewportLocker.h"
 
-
+#include "Engine/Engine.h"
+#include "Engine/GameInstance.h"
 #if WITH_EDITOR
 #include "Modules/ModuleManager.h"
 #include "LevelEditor.h"
