Fix incorrect forward declaration

diff --git a/Engine/Plugins/Experimental/VirtualCamera/Source/VCamCoreEditor/Public/VCamEditorLibrary.h b/Engine/Plugins/Experimental/VirtualCamera/Source/VCamCoreEditor/Public/VCamEditorLibrary.h
--- a/Engine/Plugins/Experimental/VirtualCamera/Source/VCamCoreEditor/Public/VCamEditorLibrary.h
+++ b/Engine/Plugins/Experimental/VirtualCamera/Source/VCamCoreEditor/Public/VCamEditorLibrary.h
@@ -4,10 +4,9 @@
 
 #include "CoreMinimal.h"
 #include "Kismet/BlueprintFunctionLibrary.h"
-
 #include "VCamEditorLibrary.generated.h"
 
-class UVcamComponent;
+class UVCamComponent;
 
 UCLASS()
 class UVCamEditorLibrary : public UBlueprintFunctionLibrary
