Move VCamCoreEditor factories to Private folder

diff --git a/Engine/Plugins/Experimental/VirtualCamera/Source/VCamCoreEditor/Private/Factories/VCamModifierFactory.h b/Engine/Plugins/Experimental/VirtualCamera/Source/VCamCoreEditor/Private/Factories/VCamModifierFactory.h
new file mode 100644
--- /dev/null
+++ b/Engine/Plugins/Experimental/VirtualCamera/Source/VCamCoreEditor/Private/Factories/VCamModifierFactory.h
@@ -0,0 +1,24 @@
+// Copyright Epic Games, Inc. All Rights Reserved.
+
+#pragma once
+
+#include "CoreMinimal.h"
+#include "VCamObjectWithInputFactory.h"
+
+#include "VCamModifierFactory.generated.h"
+
+UCLASS()
+class VCAMCOREEDITOR_API UVCamModifierFactory : public UVCamObjectWithInputFactory
+{
+	GENERATED_BODY()
+public:
+	UVCamModifierFactory();
+	
+	//~ Begin UFactory Interface
+	virtual FText GetDisplayName() const override;
+	virtual FText GetToolTip() const override;
+	virtual UObject* FactoryCreateNew(UClass* Class, UObject* InParent, FName Name, EObjectFlags Flags, UObject* Context, FFeedbackContext* Warn) override;
+	virtual uint32 GetMenuCategories() const override;
+	//~ Begin UFactory Interface
+	
+};
\ No newline at end of file
diff --git a/Engine/Plugins/Experimental/VirtualCamera/Source/VCamCoreEditor/Private/Factories/VCamObjectWithInputFactory.h b/Engine/Plugins/Experimental/VirtualCamera/Source/VCamCoreEditor/Private/Factories/VCamObjectWithInputFactory.h
new file mode 100644
--- /dev/null
+++ b/Engine/Plugins/Experimental/VirtualCamera/Source/VCamCoreEditor/Private/Factories/VCamObjectWithInputFactory.h
@@ -0,0 +1,29 @@
+// Copyright Epic Games, Inc. All Rights Reserved.
+
+#pragma once
+
+#include "CoreMinimal.h"
+#include "Factories/Factory.h"
+#include "InputMappingContext.h"
+
+#include "VCamObjectWithInputFactory.generated.h"
+
+UCLASS(hidecategories = Object)
+class VCAMCOREEDITOR_API UVCamObjectWithInputFactory : public UFactory
+{
+	GENERATED_BODY()
+public:	
+	// The parent class of the created blueprint
+	UPROPERTY(EditAnywhere, Category="BlueprintVirtualSubjectFactory", meta=(AllowAbstract = "", BlueprintBaseOnly = ""))
+	TSubclassOf<UObject> ParentClass;
+
+	UPROPERTY(BlueprintReadWrite, Category = "VCam Modifier Factory")
+	TObjectPtr<UInputMappingContext> InputMappingContext;
+
+	UPROPERTY(BlueprintReadWrite, Category = "VCam Modifier Factory")
+	int32 MappingPriority;
+
+	// Begin UFactory Interface
+	virtual bool ConfigureProperties() override;
+	// End UFactory Interface
+};
\ No newline at end of file
diff --git a/Engine/Plugins/Experimental/VirtualCamera/Source/VCamCoreEditor/Private/Factories/VCamStateSwitcherWidgetFactory.h b/Engine/Plugins/Experimental/VirtualCamera/Source/VCamCoreEditor/Private/Factories/VCamStateSwitcherWidgetFactory.h
new file mode 100644
--- /dev/null
+++ b/Engine/Plugins/Experimental/VirtualCamera/Source/VCamCoreEditor/Private/Factories/VCamStateSwitcherWidgetFactory.h
@@ -0,0 +1,21 @@
+// Copyright Epic Games, Inc. All Rights Reserved.
+
+#pragma once
+
+#include "CoreMinimal.h"
+#include "VCamWidgetFactory.h"
+#include "VCamStateSwitcherWidgetFactory.generated.h"
+
+UCLASS()
+class VCAMCOREEDITOR_API UVCamStateSwitcherWidgetFactory : public UVCamWidgetFactory
+{
+	GENERATED_BODY()
+public:
+	
+	UVCamStateSwitcherWidgetFactory();
+
+	//~ Begin UFactory Interface
+	virtual FText GetDisplayName() const override;
+	virtual FText GetToolTip() const override;
+	//~ End UFactory Interface
+};
\ No newline at end of file
diff --git a/Engine/Plugins/Experimental/VirtualCamera/Source/VCamCoreEditor/Private/Factories/VCamWidgetFactory.h b/Engine/Plugins/Experimental/VirtualCamera/Source/VCamCoreEditor/Private/Factories/VCamWidgetFactory.h
new file mode 100644
--- /dev/null
+++ b/Engine/Plugins/Experimental/VirtualCamera/Source/VCamCoreEditor/Private/Factories/VCamWidgetFactory.h
@@ -0,0 +1,24 @@
+// Copyright Epic Games, Inc. All Rights Reserved.
+
+#pragma once
+
+#include "CoreMinimal.h"
+#include "VCamObjectWithInputFactory.h"
+
+#include "VCamWidgetFactory.generated.h"
+
+UCLASS()
+class VCAMCOREEDITOR_API UVCamWidgetFactory : public UVCamObjectWithInputFactory
+{
+	GENERATED_BODY()
+public:
+	UVCamWidgetFactory();
+	
+	//~ Begin UFactory Interface
+	virtual FText GetDisplayName() const override;
+	virtual FText GetToolTip() const override;
+	virtual UObject* FactoryCreateNew(UClass* Class, UObject* InParent, FName Name, EObjectFlags Flags, UObject* Context, FFeedbackContext* Warn) override;
+	virtual uint32 GetMenuCategories() const override;
+	//~ Begin UFactory Interface
+	
+};
\ No newline at end of file
diff --git a/Engine/Plugins/Experimental/VirtualCamera/Source/VCamCoreEditor/Public/Factories/VCamModifierFactory.h b/Engine/Plugins/Experimental/VirtualCamera/Source/VCamCoreEditor/Public/Factories/VCamModifierFactory.h
deleted file mode 100644
diff --git a/Engine/Plugins/Experimental/VirtualCamera/Source/VCamCoreEditor/Public/Factories/VCamObjectWithInputFactory.h b/Engine/Plugins/Experimental/VirtualCamera/Source/VCamCoreEditor/Public/Factories/VCamObjectWithInputFactory.h
deleted file mode 100644
diff --git a/Engine/Plugins/Experimental/VirtualCamera/Source/VCamCoreEditor/Public/Factories/VCamStateSwitcherWidgetFactory.h b/Engine/Plugins/Experimental/VirtualCamera/Source/VCamCoreEditor/Public/Factories/VCamStateSwitcherWidgetFactory.h
deleted file mode 100644
diff --git a/Engine/Plugins/Experimental/VirtualCamera/Source/VCamCoreEditor/Public/Factories/VCamWidgetFactory.h b/Engine/Plugins/Experimental/VirtualCamera/Source/VCamCoreEditor/Public/Factories/VCamWidgetFactory.h
deleted file mode 100644
