Override resolution not reset when switching pixel stream viewport target

diff --git a/Engine/Plugins/Experimental/VirtualCamera/Source/PixelStreamingVCam/Private/VCamPixelStreamingSession.cpp b/Engine/Plugins/Experimental/VirtualCamera/Source/PixelStreamingVCam/Private/VCamPixelStreamingSession.cpp
--- a/Engine/Plugins/Experimental/VirtualCamera/Source/PixelStreamingVCam/Private/VCamPixelStreamingSession.cpp
+++ b/Engine/Plugins/Experimental/VirtualCamera/Source/PixelStreamingVCam/Private/VCamPixelStreamingSession.cpp
@@ -275,9 +275,6 @@ void UVCamPixelStreamingSession::StartCapture()
 		TWeakPtr<FSceneViewport> SceneViewport = GetTargetSceneViewport();
 		if (TSharedPtr<FSceneViewport> PinnedSceneViewport = SceneViewport.Pin())
 		{
-			// Apply the override resolution if applicable
-			UpdateOverrideResolution(bUseOverrideResolution, PinnedSceneViewport);
-			
 			MediaCapture->CaptureSceneViewport(PinnedSceneViewport, Options);
 			UE_LOG(LogPixelStreamingVCam, Log, TEXT("PixelStreaming set to capture scene viewport."));
 		}
@@ -330,9 +327,6 @@ void UVCamPixelStreamingSession::Deactivate()
 		StopSignallingServer();
 	}
 
-	// Remove the override resolution
-	UpdateOverrideResolution(false);
-
 	Super::Deactivate();
 	if (bUsingDummyUMG)
 	{
@@ -357,42 +351,14 @@ void UVCamPixelStreamingSession::PostEditChangeProperty(FPropertyChangedEvent& P
 	if (Property && PropertyChangedEvent.ChangeType != EPropertyChangeType::Interactive)
 	{
 		static FName NAME_FromComposureOutputProviderIndex = GET_MEMBER_NAME_CHECKED(UVCamPixelStreamingSession, FromComposureOutputProviderIndex);
-		static FName NAME_UseOverrideResolution = GET_MEMBER_NAME_CHECKED(UVCamPixelStreamingSession, bUseOverrideResolution);
-
 		const FName PropertyName = Property->GetFName();
-		
 		if (PropertyName == NAME_FromComposureOutputProviderIndex)
 		{
 			SetActive(false);
 		}
-		else if (PropertyName == NAME_UseOverrideResolution)
-		{
-			UpdateOverrideResolution(bUseOverrideResolution);
-		}
 	}
 	Super::PostEditChangeProperty(PropertyChangedEvent);
 }
 #endif
 
-void UVCamPixelStreamingSession::UpdateOverrideResolution(bool bApplyOverride) const
-{
-	if (const TSharedPtr<FSceneViewport> TargetSceneViewport = GetTargetSceneViewport())
-	{
-		UpdateOverrideResolution(bApplyOverride, TargetSceneViewport);
-	}
-}
-
-void UVCamPixelStreamingSession::UpdateOverrideResolution(bool bApplyOverride, const TSharedPtr<FSceneViewport>& SceneViewport) const
-{
-	// Apply the override resolution if requested otherwise remove any existing override
-	if (bApplyOverride)
-	{
-		SceneViewport->SetFixedViewportSize(OverrideResolution.X, OverrideResolution.Y);
-	}
-	else
-	{
-		SceneViewport->SetFixedViewportSize(0, 0);
-	}
-}
-
 IMPLEMENT_MODULE(FDefaultModuleImpl, PixelStreamingVCam)
diff --git a/Engine/Plugins/Experimental/VirtualCamera/Source/PixelStreamingVCam/Public/VCamPixelStreamingSession.h b/Engine/Plugins/Experimental/VirtualCamera/Source/PixelStreamingVCam/Public/VCamPixelStreamingSession.h
--- a/Engine/Plugins/Experimental/VirtualCamera/Source/PixelStreamingVCam/Public/VCamPixelStreamingSession.h
+++ b/Engine/Plugins/Experimental/VirtualCamera/Source/PixelStreamingVCam/Public/VCamPixelStreamingSession.h
@@ -48,6 +48,10 @@ protected:
 	UPROPERTY(Transient)
 	TObjectPtr<UPixelStreamingMediaCapture> MediaCapture = nullptr;
 
+	//~ Begin UVCamOutputProviderBase Interface
+	virtual bool ShouldOverrideResolutionOnActivationEvents() const { return true; }
+	//~ End UVCamOutputProviderBase Interface
+
 private:
 	void SetupSignallingServer();
 	void StopSignallingServer();
@@ -56,22 +60,6 @@ private:
 	void SetupCustomInputHandling();
 	void OnCaptureStateChanged();
 	void OnARKitTransformReceived(FPixelStreamingPlayerId PlayerId, uint8 Type, TArray<uint8> Data);
-	
-	/**
-	 * Applies or removes a fixed resolution to a viewport.
-	 * As no viewport is provided this version will try to get the Target Viewport from the VCam Component and use that as the viewport to update
-	 *
-	 * @param bApplyOverride	If true applies the override resolution settings to the viewport otherwise removes any existing override
-	 */
-	void UpdateOverrideResolution(bool bApplyOverride) const;
-
-	/**
-	 * Applies or removes a fixed resolution to a given viewport 
-	 *
-	 * @param bApplyOverride	If true applies the override resolution settings to the viewport otherwise removes any existing override
-	 * @param SceneViewport		The viewport to update the override resolution on. The function assumes this pointer has been validated already and will not test it before use.
-	 */
-	void UpdateOverrideResolution(bool bApplyOverride, const TSharedPtr<FSceneViewport>& SceneViewport) const;
 private:
 	FHitResult 	LastViewportTouchResult;
 	bool 		bUsingDummyUMG = false;
diff --git a/Engine/Plugins/Experimental/VirtualCamera/Source/VCamCore/Private/Output/VCamOutputProviderBase.cpp b/Engine/Plugins/Experimental/VirtualCamera/Source/VCamCore/Private/Output/VCamOutputProviderBase.cpp
--- a/Engine/Plugins/Experimental/VirtualCamera/Source/VCamCore/Private/Output/VCamOutputProviderBase.cpp
+++ b/Engine/Plugins/Experimental/VirtualCamera/Source/VCamCore/Private/Output/VCamOutputProviderBase.cpp
@@ -12,6 +12,7 @@
 #include "Blueprint/WidgetTree.h"
 #include "Engine/Engine.h"
 #include "Framework/Application/SlateApplication.h"
+#include "Slate/SceneViewport.h"
 #include "UObject/UObjectBaseUtility.h"
 
 #if WITH_EDITOR
@@ -78,10 +79,19 @@ void UVCamOutputProviderBase::Activate()
 {
 	CreateUMG();
 	DisplayUMG();
+
+	if (ShouldOverrideResolutionOnActivationEvents() && bUseOverrideResolution)
+	{
+		ApplyOverrideResolutionForViewport(TargetViewport);
+	}
 }
 
 void UVCamOutputProviderBase::Deactivate()
 {
+	if (ShouldOverrideResolutionOnActivationEvents())
+	{
+		RestoreOverrideResolutionForViewport(TargetViewport);
+	}
 	DestroyUMG();
 }
 
@@ -152,6 +162,22 @@ void UVCamOutputProviderBase::CreateUMG()
 	UE_LOG(LogVCamOutputProvider, Log, TEXT("CreateUMG widget named %s from class %s"), *UMGWidget->GetName(), *UMGWidget->WidgetClass->GetName());
 }
 
+void UVCamOutputProviderBase::RestoreOverrideResolutionForViewport(EVCamTargetViewportID ViewportToRestore)
+{
+	if (const TSharedPtr<FSceneViewport> TargetSceneViewport = GetSceneViewport(ViewportToRestore))
+	{
+		TargetSceneViewport->SetFixedViewportSize(0, 0);
+	}
+}
+
+void UVCamOutputProviderBase::ApplyOverrideResolutionForViewport(EVCamTargetViewportID Viewport)
+{
+	if (const TSharedPtr<FSceneViewport> TargetSceneViewport = GetSceneViewport(Viewport))
+	{
+		TargetSceneViewport->SetFixedViewportSize(OverrideResolution.X, OverrideResolution.Y);
+	}
+}
+
 void UVCamOutputProviderBase::DisplayUMG()
 {
 	if (UMGWidget)
@@ -277,7 +303,7 @@ UVCamOutputProviderBase* UVCamOutputProviderBase::GetOtherOutputProviderByIndex(
 	return nullptr;
 }
 
-TSharedPtr<FSceneViewport> UVCamOutputProviderBase::GetTargetSceneViewport() const
+TSharedPtr<FSceneViewport> UVCamOutputProviderBase::GetSceneViewport(EVCamTargetViewportID InTargetViewport) const
 {
 	TSharedPtr<FSceneViewport> SceneViewport;
 
@@ -307,7 +333,11 @@ TSharedPtr<FSceneViewport> UVCamOutputProviderBase::GetTargetSceneViewport() con
 			}
 			else if (Context.WorldType == EWorldType::Editor)
 			{
-				if (FLevelEditorViewportClient* LevelViewportClient = GetTargetLevelViewportClient())
+				TSharedPtr<SLevelViewport> Viewport = UE::VCamCore::LevelViewportUtils::Private::GetLevelViewport(InTargetViewport);
+				FLevelEditorViewportClient* LevelViewportClient = Viewport
+					? &Viewport->GetLevelViewportClient()
+					: nullptr;
+				if (LevelViewportClient)
 				{
 					TSharedPtr<SEditorViewport> ViewportWidget = LevelViewportClient->GetEditorViewportWidget();
 					if (ViewportWidget.IsValid())
@@ -402,12 +432,16 @@ void UVCamOutputProviderBase::PostEditChangeProperty(FPropertyChangedEvent& Prop
 	{
 		static FName NAME_IsActive = GET_MEMBER_NAME_CHECKED(UVCamOutputProviderBase, bIsActive);
 		static FName NAME_UMGClass = GET_MEMBER_NAME_CHECKED(UVCamOutputProviderBase, UMGClass);
+		static FName NAME_TargetViewport = GET_MEMBER_NAME_CHECKED(UVCamOutputProviderBase, TargetViewport);
+		static FName NAME_OverrideResolution = GET_MEMBER_NAME_CHECKED(UVCamOutputProviderBase, OverrideResolution);
+		static FName NAME_bUseOverrideResolution = GET_MEMBER_NAME_CHECKED(UVCamOutputProviderBase, bUseOverrideResolution);
 
-		if (Property->GetFName() == NAME_IsActive)
+		const FName PropertyName = Property->GetFName();
+		if (PropertyName == NAME_IsActive)
 		{
 			SetActive(bIsActive);
 		}
-		else if (Property->GetFName() == NAME_UMGClass)
+		else if (PropertyName == NAME_UMGClass)
 		{
 			if (bIsActive)
 			{
@@ -415,6 +449,22 @@ void UVCamOutputProviderBase::PostEditChangeProperty(FPropertyChangedEvent& Prop
 				SetActive(true);
 			}
 		}
+		else if (PropertyName == NAME_TargetViewport)
+		{
+			for (int32 i = 0; i < static_cast<int32>(EVCamTargetViewportID::Count); ++i)
+			{
+				RestoreOverrideResolutionForViewport(static_cast<EVCamTargetViewportID>(i));
+			}
+			ApplyOverrideResolutionForViewport(TargetViewport);
+		}
+		else if (PropertyName == NAME_OverrideResolution)
+		{
+			ApplyOverrideResolutionForViewport(TargetViewport);
+		}
+		else if (PropertyName == NAME_bUseOverrideResolution)
+		{
+			RestoreOverrideResolutionForViewport(TargetViewport);
+		}
 	}
 
 	Super::PostEditChangeProperty(PropertyChangedEvent);
diff --git a/Engine/Plugins/Experimental/VirtualCamera/Source/VCamCore/Private/Output/VCamOutputRemoteSession.cpp b/Engine/Plugins/Experimental/VirtualCamera/Source/VCamCore/Private/Output/VCamOutputRemoteSession.cpp
--- a/Engine/Plugins/Experimental/VirtualCamera/Source/VCamCore/Private/Output/VCamOutputRemoteSession.cpp
+++ b/Engine/Plugins/Experimental/VirtualCamera/Source/VCamCore/Private/Output/VCamOutputRemoteSession.cpp
@@ -52,7 +52,6 @@ void UVCamOutputRemoteSession::Activate()
 	}
 
 	CreateRemoteSession();
-	
 	Super::Activate();
 }
 
@@ -100,17 +99,7 @@ void UVCamOutputRemoteSession::CreateRemoteSession()
 			if (RemoteSessionHost.IsValid())
 			{
 				RemoteSessionHost->RegisterChannelChangeDelegate(FOnRemoteSessionChannelChange::FDelegate::CreateUObject(this, &UVCamOutputRemoteSession::OnRemoteSessionChannelChange));
-
 				RemoteSessionHost->Tick(0.0f);
-
-				if (bUseOverrideResolution)
-				{
-					TSharedPtr<FSceneViewport> SceneViewport = GetTargetSceneViewport();
-					if (SceneViewport.IsValid())
-					{
-						SceneViewport->SetFixedViewportSize(OverrideResolution.X, OverrideResolution.Y);
-					}
-				}
 			}
 			else
 			{
@@ -161,15 +150,6 @@ void UVCamOutputRemoteSession::DestroyRemoteSession()
 			}
 		}
 #endif
-
-		if (bUseOverrideResolution)
-		{
-			TSharedPtr<FSceneViewport> SceneViewport = GetTargetSceneViewport();
-			if (SceneViewport.IsValid())
-			{
-				SceneViewport->SetFixedViewportSize(0, 0);
-			}
-		}
 	}
 }
 
diff --git a/Engine/Plugins/Experimental/VirtualCamera/Source/VCamCore/Public/EVCamTargetViewportID.h b/Engine/Plugins/Experimental/VirtualCamera/Source/VCamCore/Public/EVCamTargetViewportID.h
--- a/Engine/Plugins/Experimental/VirtualCamera/Source/VCamCore/Public/EVCamTargetViewportID.h
+++ b/Engine/Plugins/Experimental/VirtualCamera/Source/VCamCore/Public/EVCamTargetViewportID.h
@@ -12,5 +12,7 @@ enum class EVCamTargetViewportID : uint8
 	Viewport1 = 1,
 	Viewport2 = 2,
 	Viewport3 = 3,
-	Viewport4 = 4
+	Viewport4 = 4,
+
+	Count UMETA(Hidden)
 };
\ No newline at end of file
diff --git a/Engine/Plugins/Experimental/VirtualCamera/Source/VCamCore/Public/Output/VCamOutputProviderBase.h b/Engine/Plugins/Experimental/VirtualCamera/Source/VCamCore/Public/Output/VCamOutputProviderBase.h
--- a/Engine/Plugins/Experimental/VirtualCamera/Source/VCamCore/Public/Output/VCamOutputProviderBase.h
+++ b/Engine/Plugins/Experimental/VirtualCamera/Source/VCamCore/Public/Output/VCamOutputProviderBase.h
@@ -129,12 +129,22 @@ protected:
 
 	virtual void CreateUMG();
 
+	/** Whether this subclass supports override resolutions. */
+	virtual bool ShouldOverrideResolutionOnActivationEvents() const { return false; }
+	/** Removes the override resolution from the given viewport. */
+	void RestoreOverrideResolutionForViewport(EVCamTargetViewportID ViewportToRestore);
+	/** Applies OverrideResolution to the passed in viewport - bUseOverrideResolution was already checked. */
+	void ApplyOverrideResolutionForViewport(EVCamTargetViewportID Viewport);
+
 	void DisplayUMG();
 	void DestroyUMG();
 
 	UVCamOutputProviderBase* GetOtherOutputProviderByIndex(int32 Index) const;
 
-	TSharedPtr<FSceneViewport> GetTargetSceneViewport() const;
+	/** Gets the scene viewport identified by the currently configured TargetViewport. */
+	TSharedPtr<FSceneViewport> GetTargetSceneViewport() const { return GetSceneViewport(TargetViewport); }
+	/** Gets the viewport identified by the passed in parameters. */
+	TSharedPtr<FSceneViewport> GetSceneViewport(EVCamTargetViewportID InTargetViewport) const;
 	TWeakPtr<SWindow> GetTargetInputWindow() const;
 
 #if WITH_EDITOR
diff --git a/Engine/Plugins/Experimental/VirtualCamera/Source/VCamCore/Public/Output/VCamOutputRemoteSession.h b/Engine/Plugins/Experimental/VirtualCamera/Source/VCamCore/Public/Output/VCamOutputRemoteSession.h
--- a/Engine/Plugins/Experimental/VirtualCamera/Source/VCamCore/Public/Output/VCamOutputRemoteSession.h
+++ b/Engine/Plugins/Experimental/VirtualCamera/Source/VCamCore/Public/Output/VCamOutputRemoteSession.h
@@ -48,6 +48,10 @@ protected:
 	UPROPERTY(Transient)
 	TObjectPtr<URemoteSessionMediaCapture> MediaCapture = nullptr;
 
+	//~ Begin UVCamOutputProviderBase Interface
+	virtual bool ShouldOverrideResolutionOnActivationEvents() const { return true; }
+	//~ End UVCamOutputProviderBase Interface
+
 private:
 	TSharedPtr<IRemoteSessionUnmanagedRole> RemoteSessionHost;
 
