Fix override resolution not getting applied

diff --git a/Engine/Plugins/Experimental/VirtualCamera/Source/VCamCore/Private/Output/VCamOutputProviderBase.cpp b/Engine/Plugins/Experimental/VirtualCamera/Source/VCamCore/Private/Output/VCamOutputProviderBase.cpp
--- a/Engine/Plugins/Experimental/VirtualCamera/Source/VCamCore/Private/Output/VCamOutputProviderBase.cpp
+++ b/Engine/Plugins/Experimental/VirtualCamera/Source/VCamCore/Private/Output/VCamOutputProviderBase.cpp
@@ -154,10 +154,6 @@ void UVCamOutputProviderBase::CreateUMG()
 	UMGWidget->SetDisplayTypes(DisplayType, DisplayType, DisplayType);
 	UMGWidget->PostProcessDisplayType.bReceiveHardwareInput = true;
 
-	// The post process material created by EVPWidgetDisplayType::PostProcess should affect only our VCam, not the entire world.
-	// This allows you to have multiple virtual camera in the world affect different viewports.
-	UMGWidget->SetCustomPostProcessSettingsSource(TargetCamera.Get());
-
 #if WITH_EDITOR
 	UMGWidget->SetAllTargetViewports(GetTargetLevelViewport());
 #endif
@@ -182,6 +178,18 @@ void UVCamOutputProviderBase::ApplyOverrideResolutionForViewport(EVCamTargetView
 	}
 }
 
+void UVCamOutputProviderBase::ReapplyOverrideResolution(EVCamTargetViewportID Viewport)
+{
+	if (bUseOverrideResolution)
+	{
+		ApplyOverrideResolutionForViewport(Viewport);
+	}
+	else
+	{
+		RestoreOverrideResolutionForViewport(Viewport);
+	}
+}
+
 void UVCamOutputProviderBase::DisplayUMG()
 {
 	if (UMGWidget)
@@ -463,13 +471,9 @@ void UVCamOutputProviderBase::PostEditChangeProperty(FPropertyChangedEvent& Prop
 			}
 			ApplyOverrideResolutionForViewport(TargetViewport);
 		}
-		else if (PropertyName == NAME_OverrideResolution)
-		{
-			ApplyOverrideResolutionForViewport(TargetViewport);
-		}
-		else if (PropertyName == NAME_bUseOverrideResolution)
+		else if (PropertyName == NAME_OverrideResolution || PropertyName == NAME_bUseOverrideResolution)
 		{
-			RestoreOverrideResolutionForViewport(TargetViewport);
+			ReapplyOverrideResolution(TargetViewport);
 		}
 	}
 
diff --git a/Engine/Plugins/Experimental/VirtualCamera/Source/VCamCore/Public/Output/VCamOutputProviderBase.h b/Engine/Plugins/Experimental/VirtualCamera/Source/VCamCore/Public/Output/VCamOutputProviderBase.h
--- a/Engine/Plugins/Experimental/VirtualCamera/Source/VCamCore/Public/Output/VCamOutputProviderBase.h
+++ b/Engine/Plugins/Experimental/VirtualCamera/Source/VCamCore/Public/Output/VCamOutputProviderBase.h
@@ -136,6 +136,8 @@ protected:
 	/** Applies OverrideResolution to the passed in viewport - bUseOverrideResolution was already checked. */
 	void ApplyOverrideResolutionForViewport(EVCamTargetViewportID Viewport);
 
+	void ReapplyOverrideResolution(EVCamTargetViewportID Viewport);
+
 	void DisplayUMG();
 	void DestroyUMG();
 
