Add editor-only modifier for modifiers which depend on non-cooked content

diff --git a/Engine/Plugins/Experimental/VirtualCamera/Source/VCamCoreEditor/Public/EditorOnlyVCamModifier.h b/Engine/Plugins/Experimental/VirtualCamera/Source/VCamCoreEditor/Public/EditorOnlyVCamModifier.h
new file mode 100644
--- /dev/null
+++ b/Engine/Plugins/Experimental/VirtualCamera/Source/VCamCoreEditor/Public/EditorOnlyVCamModifier.h
@@ -0,0 +1,21 @@
+// Copyright Epic Games, Inc. All Rights Reserved.
+
+#pragma once
+
+#include "CoreMinimal.h"
+#include "Modifier/VCamModifier.h"
+#include "EditorOnlyVCamModifier.generated.h"
+
+/**
+ * A modifier that only exists in uncooked builds.
+ */
+UCLASS(Abstract)
+class VCAMCOREEDITOR_API UEditorOnlyVCamModifier : public UVCamModifier
+{
+	GENERATED_BODY()
+public:
+
+	//~ Begin UObject Interface
+	virtual bool IsEditorOnly() const override { return true; }
+	//~ End UObject Interface
+};
\ No newline at end of file
