Fix PixelStreamingVCam being loaded after Blueprints are compiled in engine init causing references to Pixel Streaming Provider to be dropped.

diff --git a/Engine/Plugins/Experimental/VirtualCamera/VirtualCamera.uplugin b/Engine/Plugins/Experimental/VirtualCamera/VirtualCamera.uplugin
--- a/Engine/Plugins/Experimental/VirtualCamera/VirtualCamera.uplugin
+++ b/Engine/Plugins/Experimental/VirtualCamera/VirtualCamera.uplugin
@@ -26,7 +26,7 @@
 		{
 			"Name": "PixelStreamingVCam",
 			"Type": "Editor",
-			"LoadingPhase": "PostEngineInit",
+			"LoadingPhase": "PostDefault",
 			"PlatformAllowList": [ "Win64", "Linux" ],
 			"TargetDenyList": [ "Server" ]
 		},
@@ -114,6 +114,10 @@
 		{
 			"Name": "PixelStreaming",
 			"Enabled": true
+		},
+		{
+			"Name" : "PixelCapture",
+			"Enabled": true
 		}
 	]
 }
