Fix crashes related due incorrect object destruction:
- crash on saving map
- other hard to repro crash

diff --git a/Engine/Plugins/Experimental/VirtualCamera/Source/VCamCore/Private/Modifier/VCamModifier.cpp b/Engine/Plugins/Experimental/VirtualCamera/Source/VCamCore/Private/Modifier/VCamModifier.cpp
--- a/Engine/Plugins/Experimental/VirtualCamera/Source/VCamCore/Private/Modifier/VCamModifier.cpp
+++ b/Engine/Plugins/Experimental/VirtualCamera/Source/VCamCore/Private/Modifier/VCamModifier.cpp
@@ -22,11 +22,12 @@ void UVCamBlueprintModifier::Initialize(UVCamModifierContext* Context, UInputCom
 
 void UVCamBlueprintModifier::Deinitialize()
 {
-	// Forward the Initialize call to the Blueprint Event only if we have a valid world
-	if (GetWorld())
+	if (GetWorld()
+		&& !HasAnyFlags(RF_BeginDestroyed | RF_FinishDestroyed)
+		// Executing BP code during this phase is not allowed - you hit this when saving the world.
+		&& !FUObjectThreadContext::Get().IsRoutingPostLoad)
 	{
 		FEditorScriptExecutionGuard ScriptGuard;
-
 		OnDeinitialize();
 	}
 	
@@ -65,6 +66,16 @@ void UVCamModifier::Deinitialize()
 	bRequiresInitialization = true;
 }
 
+void UVCamModifier::BeginDestroy()
+{
+	if (!bRequiresInitialization)
+	{
+		Deinitialize();
+	}
+	
+	UObject::BeginDestroy();
+}
+
 void UVCamModifier::PostLoad()
 {
 	Super::PostLoad();
diff --git a/Engine/Plugins/Experimental/VirtualCamera/Source/VCamCore/Private/VCamComponent.cpp b/Engine/Plugins/Experimental/VirtualCamera/Source/VCamCore/Private/VCamComponent.cpp
--- a/Engine/Plugins/Experimental/VirtualCamera/Source/VCamCore/Private/VCamComponent.cpp
+++ b/Engine/Plugins/Experimental/VirtualCamera/Source/VCamCore/Private/VCamComponent.cpp
@@ -104,12 +104,21 @@ void UVCamComponent::OnComponentDestroyed(bool bDestroyingHierarchy)
 
 	for (UVCamOutputProviderBase* Provider : OutputProviders)
 	{
-		if (Provider)
+		if (Provider && Provider->IsInitialized())
 		{
 			Provider->Deinitialize();
 		}
 	}
 
+	for (FModifierStackEntry& ModifierEntry : ModifierStack)
+	{
+		TObjectPtr<UVCamModifier>& Modifier = ModifierEntry.GeneratedModifier;
+		if (Modifier && !Modifier->IsInitialized())
+		{
+			Modifier->Deinitialize();
+		}
+	}
+
 	UnregisterInputComponent();
 
 #if WITH_EDITOR
@@ -398,6 +407,7 @@ void UVCamComponent::PostEditChangeProperty(FPropertyChangedEvent& PropertyChang
 		else if (PropertyName == NAME_ModifierStack)
 		{
 			ValidateModifierStack();
+			SavedModifierStack.Empty();
 		}
 		else if (PropertyName == NAME_InputProfile)
 		{
@@ -716,7 +726,7 @@ void UVCamComponent::Update()
 				if (ModifierStackEntry.bEnabled)
 				{
 					// Initialize the Modifier if required
-					if (Modifier->DoesRequireInitialization())
+					if (Modifier->IsInitialized())
 					{
 						Modifier->Initialize(ModifierContext, InputComponent);
 						AddInputMappingContext(Modifier);
@@ -727,7 +737,7 @@ void UVCamComponent::Update()
 				else
 				{
 					// If the modifier is initialized but not enabled then we deinitialize it
-					if (!Modifier->DoesRequireInitialization())
+					if (!Modifier->IsInitialized())
 					{
 						Modifier->Deinitialize();
 					}
diff --git a/Engine/Plugins/Experimental/VirtualCamera/Source/VCamCore/Public/Modifier/ModifierStackEntry.h b/Engine/Plugins/Experimental/VirtualCamera/Source/VCamCore/Public/Modifier/ModifierStackEntry.h
--- a/Engine/Plugins/Experimental/VirtualCamera/Source/VCamCore/Public/Modifier/ModifierStackEntry.h
+++ b/Engine/Plugins/Experimental/VirtualCamera/Source/VCamCore/Public/Modifier/ModifierStackEntry.h
@@ -50,15 +50,6 @@ struct FModifierStackEntry
 #if WITH_EDITOR
 		EditorGuid = FGuid::NewGuid();
 #endif
-	};
-
-	~FModifierStackEntry()
-	{
-		// It's possible the Generated Modifier may pass IsValid while having already been destroyed so ensure we explicitly check for this case
-		if (IsValid(GeneratedModifier) && !GeneratedModifier->HasAnyFlags(RF_BeginDestroyed | RF_FinishDestroyed))
-		{
-			GeneratedModifier->Deinitialize();
-		}
 	}
 
 	bool operator==(const FModifierStackEntry& Other) const
diff --git a/Engine/Plugins/Experimental/VirtualCamera/Source/VCamCore/Public/Modifier/VCamModifier.h b/Engine/Plugins/Experimental/VirtualCamera/Source/VCamCore/Public/Modifier/VCamModifier.h
--- a/Engine/Plugins/Experimental/VirtualCamera/Source/VCamCore/Public/Modifier/VCamModifier.h
+++ b/Engine/Plugins/Experimental/VirtualCamera/Source/VCamCore/Public/Modifier/VCamModifier.h
@@ -29,17 +29,19 @@ UCLASS(Blueprintable, Abstract, EditInlineNew)
 class VCAMCORE_API UVCamModifier : public UObject
 {
 	GENERATED_BODY()
-
 public:
-	virtual void Initialize(UVCamModifierContext* Context, UInputComponent* InputComponent = nullptr);
 	
+	virtual void Initialize(UVCamModifierContext* Context, UInputComponent* InputComponent = nullptr);
 	virtual void Deinitialize();
 
 	virtual void Apply(UVCamModifierContext* Context, UCineCameraComponent* CameraComponent, const float DeltaTime) {};
 
+	//~ Begin UObject Interface
+	virtual void BeginDestroy() override;
 	virtual void PostLoad() override;
+	//~ End UObject Interface
 
-	bool DoesRequireInitialization() const { return bRequiresInitialization; };
+	bool IsInitialized() const { return bRequiresInitialization; };
 
 	UFUNCTION(BlueprintCallable, Category="VirtualCamera")
 	UVCamComponent* GetOwningVCamComponent() const;
diff --git a/Engine/Plugins/Experimental/VirtualCamera/Source/VCamCore/Public/VCamComponent.h b/Engine/Plugins/Experimental/VirtualCamera/Source/VCamCore/Public/VCamComponent.h
--- a/Engine/Plugins/Experimental/VirtualCamera/Source/VCamCore/Public/VCamComponent.h
+++ b/Engine/Plugins/Experimental/VirtualCamera/Source/VCamCore/Public/VCamComponent.h
@@ -402,9 +402,6 @@ private:
 	
 	double LastEvaluationTime;
 
-	TWeakObjectPtr<AActor> Backup_ActorLock;
-	TWeakObjectPtr<AActor> Backup_ViewTarget;
-
 	TArray<UVCamOutputProviderBase*> SavedOutputProviders;
 	TArray<FModifierStackEntry> SavedModifierStack;
 	
