Fix activation logic:
- After editing the VCamActor BP, the Widget is not always added to the Viewport
- Sometimes the widget is not removed after deactivation, blocking input on the viewport

diff --git a/Engine/Plugins/Experimental/VirtualCamera/Source/VCamCore/Private/Output/VCamOutputProviderBase.cpp b/Engine/Plugins/Experimental/VirtualCamera/Source/VCamCore/Private/Output/VCamOutputProviderBase.cpp
--- a/Engine/Plugins/Experimental/VirtualCamera/Source/VCamCore/Private/Output/VCamOutputProviderBase.cpp
+++ b/Engine/Plugins/Experimental/VirtualCamera/Source/VCamCore/Private/Output/VCamOutputProviderBase.cpp
@@ -67,12 +67,11 @@ void UVCamOutputProviderBase::Initialize()
 
 void UVCamOutputProviderBase::Deinitialize()
 {
-	if (IsOuterComponentEnabled() && bIsActive)
+	if (bInitialized)
 	{
 		Deactivate();
+		bInitialized = false;
 	}
-	
-	bInitialized = false;
 }
 
 void UVCamOutputProviderBase::Activate()
@@ -440,9 +439,10 @@ TSharedPtr<SLevelViewport> UVCamOutputProviderBase::GetTargetLevelViewport() con
 
 void UVCamOutputProviderBase::PostEditChangeProperty(FPropertyChangedEvent& PropertyChangedEvent)
 {
-	FProperty* Property = PropertyChangedEvent.MemberProperty;
-
-	if (Property && PropertyChangedEvent.ChangeType != EPropertyChangeType::Interactive)
+	const FProperty* Property = PropertyChangedEvent.MemberProperty;
+	if (!HasAnyFlags(RF_ClassDefaultObject | RF_ArchetypeObject)
+		&& Property
+		&& PropertyChangedEvent.ChangeType != EPropertyChangeType::Interactive)
 	{
 		static FName NAME_IsActive = GET_MEMBER_NAME_CHECKED(UVCamOutputProviderBase, bIsActive);
 		static FName NAME_UMGClass = GET_MEMBER_NAME_CHECKED(UVCamOutputProviderBase, UMGClass);
