Crash on editing states in state switcher widget

diff --git a/Engine/Plugins/Experimental/VirtualCamera/Source/VCamCore/Private/LogVCamCore.cpp b/Engine/Plugins/Experimental/VirtualCamera/Source/VCamCore/Private/LogVCamCore.cpp
new file mode 100644
--- /dev/null
+++ b/Engine/Plugins/Experimental/VirtualCamera/Source/VCamCore/Private/LogVCamCore.cpp
@@ -0,0 +1,5 @@
+// Copyright Epic Games, Inc. All Rights Reserved.
+
+#include "LogVCamCore.h"
+
+DEFINE_LOG_CATEGORY(LogVCamCore);
\ No newline at end of file
diff --git a/Engine/Plugins/Experimental/VirtualCamera/Source/VCamCore/Private/LogVCamCore.h b/Engine/Plugins/Experimental/VirtualCamera/Source/VCamCore/Private/LogVCamCore.h
new file mode 100644
--- /dev/null
+++ b/Engine/Plugins/Experimental/VirtualCamera/Source/VCamCore/Private/LogVCamCore.h
@@ -0,0 +1,7 @@
+// Copyright Epic Games, Inc. All Rights Reserved.
+
+#pragma once
+
+#include "CoreMinimal.h"
+
+DECLARE_LOG_CATEGORY_EXTERN(LogVCamCore, Log, All);
\ No newline at end of file
diff --git a/Engine/Plugins/Experimental/VirtualCamera/Source/VCamCore/Private/UI/Switcher/VCamStateSwitcherWidget.cpp b/Engine/Plugins/Experimental/VirtualCamera/Source/VCamCore/Private/UI/Switcher/VCamStateSwitcherWidget.cpp
--- a/Engine/Plugins/Experimental/VirtualCamera/Source/VCamCore/Private/UI/Switcher/VCamStateSwitcherWidget.cpp
+++ b/Engine/Plugins/Experimental/VirtualCamera/Source/VCamCore/Private/UI/Switcher/VCamStateSwitcherWidget.cpp
@@ -46,21 +46,12 @@ bool UVCamStateSwitcherWidget::SetCurrentState(FName NewState, bool bForceUpdate
 #if WITH_EDITOR
 void UVCamStateSwitcherWidget::PostEditChangeProperty(FPropertyChangedEvent& PropertyChangedEvent)
 {
-	Super::PostEditChangeProperty(PropertyChangedEvent);
-
 	if (!States.Contains(DefaultState))
 	{
 		States.Add(DefaultState, {});
 	}
-
-	if (PropertyChangedEvent.MemberProperty && PropertyChangedEvent.MemberProperty->GetFName() == GET_MEMBER_NAME_CHECKED(UVCamStateSwitcherWidget, CurrentState))
-	{
-		constexpr bool bForceUpdate = true;
-		if (!SetCurrentState(CurrentState, bForceUpdate))
-		{
-			SetCurrentState(DefaultState, bForceUpdate);
-		}
-	}
+	
+	Super::PostEditChangeProperty(PropertyChangedEvent);
 }
 #endif 
 
@@ -69,5 +60,9 @@ void UVCamStateSwitcherWidget::NativePreConstruct()
 	Super::NativePreConstruct();
 
 	constexpr bool bForceUpdate = true;
-	SetCurrentState(CurrentState, bForceUpdate);
+	const bool bIsCurrentStateValid = SetCurrentState(CurrentState, bForceUpdate); 
+	if (!bIsCurrentStateValid)
+	{
+		SetCurrentState(DefaultState, bForceUpdate);
+	}
 }
diff --git a/Engine/Plugins/Experimental/VirtualCamera/Source/VCamCore/Private/UI/Switcher/WidgetConnectionConfig.cpp b/Engine/Plugins/Experimental/VirtualCamera/Source/VCamCore/Private/UI/Switcher/WidgetConnectionConfig.cpp
--- a/Engine/Plugins/Experimental/VirtualCamera/Source/VCamCore/Private/UI/Switcher/WidgetConnectionConfig.cpp
+++ b/Engine/Plugins/Experimental/VirtualCamera/Source/VCamCore/Private/UI/Switcher/WidgetConnectionConfig.cpp
@@ -3,11 +3,45 @@
 #include "UI/Switcher/WidgetConnectionConfig.h"
 
 #include "Blueprint/WidgetTree.h"
+#include "LogVCamCore.h"
 #include "UI/Switcher/VCamStateSwitcherWidget.h"
 
+#if WITH_EDITOR
+#include "BaseWidgetBlueprint.h"
+#endif
+
+namespace UE::VCamCore::Private
+{
+#if WITH_EDITOR
+	static UWidgetTree* GetWidgetTreeThroughBlueprintAsset(UUserWidget* ClassDefaultWidget)
+	{
+		UObject* Blueprint = ClassDefaultWidget->GetClass()->ClassGeneratedBy;
+		UBaseWidgetBlueprint* WidgetBlueprint = Cast<UBaseWidgetBlueprint>(Blueprint);
+		return WidgetBlueprint
+			? WidgetBlueprint->WidgetTree
+			: nullptr;
+	}
+#endif
+}
+
 UVCamWidget* FWidgetConnectionConfig::ResolveWidget(UVCamStateSwitcherWidget* OwnerWidget) const
 {
-	return !HasNoWidgetSet() && ensure(OwnerWidget->WidgetTree)
-		? OwnerWidget->WidgetTree->FindWidget<UVCamWidget>(Widget)
+	if (!OwnerWidget)
+	{
+		return nullptr;
+	}
+	
+	UWidgetTree* WidgetTree = OwnerWidget->WidgetTree;
+#if WITH_EDITOR
+	if (!WidgetTree && OwnerWidget->HasAnyFlags(RF_ClassDefaultObject))
+	{
+		WidgetTree = UE::VCamCore::Private::GetWidgetTreeThroughBlueprintAsset(OwnerWidget);
+	}
+#endif
+
+	const bool bCanFindWidget = !HasNoWidgetSet() && WidgetTree;
+	UE_CLOG(!bCanFindWidget, LogVCamCore, Warning, TEXT("Failed to get tree for widget %s"), *OwnerWidget->GetPathName());
+	return bCanFindWidget
+		? WidgetTree->FindWidget<UVCamWidget>(Widget)
 		: nullptr;
 }
diff --git a/Engine/Plugins/Experimental/VirtualCamera/Source/VCamCore/Public/UI/Switcher/VCamStateSwitcherWidget.h b/Engine/Plugins/Experimental/VirtualCamera/Source/VCamCore/Public/UI/Switcher/VCamStateSwitcherWidget.h
--- a/Engine/Plugins/Experimental/VirtualCamera/Source/VCamCore/Public/UI/Switcher/VCamStateSwitcherWidget.h
+++ b/Engine/Plugins/Experimental/VirtualCamera/Source/VCamCore/Public/UI/Switcher/VCamStateSwitcherWidget.h
@@ -58,7 +58,10 @@ private:
 	/** Executes when after the state has been changed */
 	UPROPERTY(BlueprintAssignable, Category = "Connections")
 	FChangeConnectionStateEvent OnPostStateChanged;
-	
+
+	/**
+	 * The states
+	 */
 	UPROPERTY(EditAnywhere, Category = "Connections")
 	TMap<FName, FVCamWidgetConnectionState> States { { DefaultState, {} } };
 
diff --git a/Engine/Plugins/Experimental/VirtualCamera/Source/VCamCoreEditor/Private/Customizations/StateSwitcher/WidgetConnectionConfigTypeCustomization.cpp b/Engine/Plugins/Experimental/VirtualCamera/Source/VCamCoreEditor/Private/Customizations/StateSwitcher/WidgetConnectionConfigTypeCustomization.cpp
--- a/Engine/Plugins/Experimental/VirtualCamera/Source/VCamCoreEditor/Private/Customizations/StateSwitcher/WidgetConnectionConfigTypeCustomization.cpp
+++ b/Engine/Plugins/Experimental/VirtualCamera/Source/VCamCoreEditor/Private/Customizations/StateSwitcher/WidgetConnectionConfigTypeCustomization.cpp
@@ -2,6 +2,7 @@
 
 #include "WidgetConnectionConfigTypeCustomization.h"
 
+#include "BaseWidgetBlueprint.h"
 #include "ConnectionTargetNodeBuilder.h"
 #include "Customizations/StateSwitcher/SStringSelectionComboBox.h"
 #include "LogVCamEditor.h"
@@ -82,7 +83,8 @@ namespace UE::VCamCoreEditor::Private
 		{
 			PropertyUtils->ForceRefresh();
 		}));
-		
+
+		const TWeakPtr<IPropertyUtilities> WeakPropertyUtils = CustomizationUtils.GetPropertyUtilities();
 		Row.CustomWidget()
 			.NameContent()
 			[
@@ -100,7 +102,7 @@ namespace UE::VCamCoreEditor::Private
 					})
 					// Property list is only build once: when the Blueprint is recompiled,
 					// the details view SHOULD refresh automatically and thus all of this is reconstructed. 
-					.ItemList(GetPropertyItemList(CustomizationUtils))
+					.ItemList(this, &FWidgetConnectionConfigTypeCustomization::GetPropertyItemList, WeakPropertyUtils)
 					.OnItemSelected_Lambda([WidgetReferencePropertyHandle](const FString& SelectedItem)
 					{
 						WidgetReferencePropertyHandle->SetValue(SelectedItem);
@@ -109,12 +111,18 @@ namespace UE::VCamCoreEditor::Private
 			];
 	}
 	
-	TArray<FString> FWidgetConnectionConfigTypeCustomization::GetPropertyItemList(IPropertyTypeCustomizationUtils& CustomizationUtils) const
+	TArray<FString> FWidgetConnectionConfigTypeCustomization::GetPropertyItemList(TWeakPtr<IPropertyUtilities> WeakPropertyUtils) const
 	{
+		TSharedPtr<IPropertyUtilities> PropertyUtilities = WeakPropertyUtils.Pin();
+		if (!PropertyUtilities)
+		{
+			return {};
+		}
+		
 		TArray<FString> PropertyItemList;
 		PropertyItemList.Add(FName(NAME_None).ToString());
 		
-		const TArray<TWeakObjectPtr<UObject>>& SelectedObjects = CustomizationUtils.GetPropertyUtilities()->GetSelectedObjects();
+		const TArray<TWeakObjectPtr<UObject>>& SelectedObjects = PropertyUtilities->GetSelectedObjects();
 		if (!ensure(SelectedObjects.Num() == 1) || !SelectedObjects[0].IsValid())
 		{
 			return PropertyItemList;
@@ -126,8 +134,18 @@ namespace UE::VCamCoreEditor::Private
 		{
 			return {};
 		}
+
+		// Two use cases are expected: Editing an instance from 1. UMG Designer tab 2. Class Defaults in Graph tab
+		// When editing Class Defaults, the WidgetTree is expected to be nullptr when editing class defaults so let's grab it from the asset.
+		const TObjectPtr<UWidgetTree> WidgetTree = Widget->WidgetTree
+			? Widget->WidgetTree
+			: GetWidgetTreeThroughBlueprintAsset(Widget);
+		if (!WidgetTree)
+		{
+			return {};
+		}
 		
-		Widget->WidgetTree->ForEachWidget([&PropertyItemList](UWidget* Widget)
+		WidgetTree->ForEachWidget([&PropertyItemList](UWidget* Widget)
 		{
 			const bool bIsVCamWidgetProperty = Widget && Widget->GetClass()->IsChildOf(UVCamWidget::StaticClass());
 			if (bIsVCamWidgetProperty)
@@ -138,6 +156,20 @@ namespace UE::VCamCoreEditor::Private
 		return PropertyItemList;
 	}
 
+	TObjectPtr<UWidgetTree> FWidgetConnectionConfigTypeCustomization::GetWidgetTreeThroughBlueprintAsset(UUserWidget* ClassDefaultWidget)
+	{
+		if (!ensure(ClassDefaultWidget->HasAnyFlags(RF_ClassDefaultObject)))
+		{
+			return nullptr;
+		}
+		
+		UObject* Blueprint = ClassDefaultWidget->GetClass()->ClassGeneratedBy;
+		UBaseWidgetBlueprint* WidgetBlueprint = Cast<UBaseWidgetBlueprint>(Blueprint);
+		return WidgetBlueprint
+			? WidgetBlueprint->WidgetTree
+			: nullptr;
+	}
+
 	void FWidgetConnectionConfigTypeCustomization::CustomizeConnectionTargetsReferenceProperty(
 		TSharedRef<IPropertyHandle> StructPropertyHandle,
 		TSharedRef<IPropertyHandle> ConnectionTargetsPropertyHandle,
diff --git a/Engine/Plugins/Experimental/VirtualCamera/Source/VCamCoreEditor/Private/Customizations/StateSwitcher/WidgetConnectionConfigTypeCustomization.h b/Engine/Plugins/Experimental/VirtualCamera/Source/VCamCoreEditor/Private/Customizations/StateSwitcher/WidgetConnectionConfigTypeCustomization.h
--- a/Engine/Plugins/Experimental/VirtualCamera/Source/VCamCoreEditor/Private/Customizations/StateSwitcher/WidgetConnectionConfigTypeCustomization.h
+++ b/Engine/Plugins/Experimental/VirtualCamera/Source/VCamCoreEditor/Private/Customizations/StateSwitcher/WidgetConnectionConfigTypeCustomization.h
@@ -5,6 +5,8 @@
 #include "CoreMinimal.h"
 #include "IPropertyTypeCustomization.h"
 
+class UWidgetTree;
+class UUserWidget;
 class IDetailPropertyRow;
 
 namespace UE::VCamCoreEditor::Private
@@ -24,7 +26,8 @@ namespace UE::VCamCoreEditor::Private
 	private:
 
 		void CustomizeWidgetReferenceProperty(TSharedRef<IPropertyHandle> WidgetReferencePropertyHandle, IDetailPropertyRow& Row, IPropertyTypeCustomizationUtils& CustomizationUtils) const;
-		TArray<FString> GetPropertyItemList(IPropertyTypeCustomizationUtils& CustomizationUtils) const;
+		TArray<FString> GetPropertyItemList(TWeakPtr<IPropertyUtilities> WeakPropertyUtils) const;
+		static TObjectPtr<UWidgetTree> GetWidgetTreeThroughBlueprintAsset(UUserWidget* ClassDefaultWidget); 
 
 		void CustomizeConnectionTargetsReferenceProperty(TSharedRef<IPropertyHandle> StructPropertyHandle, TSharedRef<IPropertyHandle> ConnectionTargetsPropertyHandle, IDetailChildrenBuilder& ChildBuilder, IPropertyTypeCustomizationUtils& CustomizationUtils) const;
 		TAttribute<TArray<FName>> CreateGetConnectionsFromChildWidgetAttribute(TSharedRef<IPropertyHandle> StructPropertyHandle, IPropertyTypeCustomizationUtils& CustomizationUtils) const;
