Allow output providers to override connection points on instanced widgets (once the widget has been created) from details panel

diff --git a/Engine/Plugins/Experimental/VirtualCamera/Source/VCamCore/Private/Output/VCamOutputProviderBase.cpp b/Engine/Plugins/Experimental/VirtualCamera/Source/VCamCore/Private/Output/VCamOutputProviderBase.cpp
--- a/Engine/Plugins/Experimental/VirtualCamera/Source/VCamCore/Private/Output/VCamOutputProviderBase.cpp
+++ b/Engine/Plugins/Experimental/VirtualCamera/Source/VCamCore/Private/Output/VCamOutputProviderBase.cpp
@@ -33,6 +33,14 @@ namespace UE::VCamCore::Private
 	static const FName LevelEditorName(TEXT("LevelEditor"));
 }
 
+UVCamOutputProviderBase::UVCamOutputProviderBase()
+{
+	OnActivatedDelegate.AddLambda([this](bool bNewValue)
+	{
+		OnActivatedDelegate_Blueprint.Broadcast(bNewValue);
+	});
+}
+
 void UVCamOutputProviderBase::BeginDestroy()
 {
 	Deinitialize();
@@ -83,6 +91,8 @@ void UVCamOutputProviderBase::Activate()
 	{
 		ApplyOverrideResolutionForViewport(TargetViewport);
 	}
+	
+	OnActivatedDelegate.Broadcast(true);
 }
 
 void UVCamOutputProviderBase::Deactivate()
@@ -91,7 +101,10 @@ void UVCamOutputProviderBase::Deactivate()
 	{
 		RestoreOverrideResolutionForViewport(TargetViewport);
 	}
+	
 	DestroyUMG();
+	
+	OnActivatedDelegate.Broadcast(false);
 }
 
 void UVCamOutputProviderBase::Tick(const float DeltaTime)
@@ -470,6 +483,12 @@ void UVCamOutputProviderBase::PostEditChangeProperty(FPropertyChangedEvent& Prop
 				RestoreOverrideResolutionForViewport(static_cast<EVCamTargetViewportID>(i));
 			}
 			ApplyOverrideResolutionForViewport(TargetViewport);
+
+			if (bIsActive)
+			{
+				SetActive(false);
+				SetActive(true);
+			}
 		}
 		else if (PropertyName == NAME_OverrideResolution || PropertyName == NAME_bUseOverrideResolution)
 		{
diff --git a/Engine/Plugins/Experimental/VirtualCamera/Source/VCamCore/Private/UI/Switcher/VCamStateSwitcherWidget.cpp b/Engine/Plugins/Experimental/VirtualCamera/Source/VCamCore/Private/UI/Switcher/VCamStateSwitcherWidget.cpp
--- a/Engine/Plugins/Experimental/VirtualCamera/Source/VCamCore/Private/UI/Switcher/VCamStateSwitcherWidget.cpp
+++ b/Engine/Plugins/Experimental/VirtualCamera/Source/VCamCore/Private/UI/Switcher/VCamStateSwitcherWidget.cpp
@@ -43,6 +43,23 @@ bool UVCamStateSwitcherWidget::SetCurrentState(FName NewState, bool bForceUpdate
 	return true;
 }
 
+TArray<FName> UVCamStateSwitcherWidget::GetStates() const
+{
+	TArray<FName> Keys;
+	States.GenerateKeyArray(Keys);
+	return Keys;
+}
+
+bool UVCamStateSwitcherWidget::GetStateInfo(FName State, FVCamWidgetConnectionState& OutStateInfo) const
+{
+	if (const FVCamWidgetConnectionState* StateInfo = States.Find(State))
+	{
+		OutStateInfo = *StateInfo;
+		return true;
+	}
+	return false;
+}
+
 #if WITH_EDITOR
 void UVCamStateSwitcherWidget::PostEditChangeProperty(FPropertyChangedEvent& PropertyChangedEvent)
 {
diff --git a/Engine/Plugins/Experimental/VirtualCamera/Source/VCamCore/Private/UI/VCamWidget.cpp b/Engine/Plugins/Experimental/VirtualCamera/Source/VCamCore/Private/UI/VCamWidget.cpp
--- a/Engine/Plugins/Experimental/VirtualCamera/Source/VCamCore/Private/UI/VCamWidget.cpp
+++ b/Engine/Plugins/Experimental/VirtualCamera/Source/VCamCore/Private/UI/VCamWidget.cpp
@@ -2,6 +2,16 @@
 
 #include "UI/VCamWidget.h"
 
+UVCamWidget::UVCamWidget(const FObjectInitializer& ObjectInitializer)
+	: Super(ObjectInitializer)
+{
+	// We have two events so C++ lambdas can be used since dynamic delegates do not support it.
+	OnPostConnectionsReinitializedDelegate.AddLambda([this]()
+	{
+		OnPostConnectionsReinitializedDelegate_Blueprint.Broadcast();
+	});
+}
+
 void UVCamWidget::NativeDestruct()
 {
 	Super::NativeDestruct();
@@ -58,6 +68,7 @@ bool UVCamWidget::ReinitializeConnections()
 	}
 
 	PostConnectionsInitialized();
+	OnPostConnectionsReinitializedDelegate.Broadcast();
 	return true;
 }
 
diff --git a/Engine/Plugins/Experimental/VirtualCamera/Source/VCamCore/Public/Output/VCamOutputProviderBase.h b/Engine/Plugins/Experimental/VirtualCamera/Source/VCamCore/Public/Output/VCamOutputProviderBase.h
--- a/Engine/Plugins/Experimental/VirtualCamera/Source/VCamCore/Public/Output/VCamOutputProviderBase.h
+++ b/Engine/Plugins/Experimental/VirtualCamera/Source/VCamCore/Public/Output/VCamOutputProviderBase.h
@@ -25,6 +25,13 @@ class VCAMCORE_API UVCamOutputProviderBase : public UObject
 	GENERATED_BODY()
 public:
 
+	DECLARE_MULTICAST_DELEGATE_OneParam(FActivationDelegate, bool /*bNewIsActive*/);
+	DECLARE_DYNAMIC_MULTICAST_DELEGATE_OneParam(FActivationDelegate_Blueprint, bool, bNewIsActive);
+	FActivationDelegate OnActivatedDelegate; 
+	/** Called when the activation state of this output provider changes. */
+	UPROPERTY(BlueprintAssignable, meta = (DisplayName = "OnActivated"))
+	FActivationDelegate_Blueprint OnActivatedDelegate_Blueprint;
+
 	/** Override the default output resolution with a custom value - NOTE you must toggle bIsActive off then back on for this to take effect */
 	UPROPERTY(EditAnywhere, BlueprintReadWrite, Category = "Output", meta = (DisplayPriority = "5"))
 	bool bUseOverrideResolution = false;
@@ -33,6 +40,8 @@ public:
 	UPROPERTY(EditAnywhere, BlueprintReadWrite, Category = "Output", meta = (DisplayPriority = "6"), meta = (EditCondition = "bUseOverrideResolution", ClampMin = 1))
 	FIntPoint OverrideResolution = { 2048, 1536 };
 
+	UVCamOutputProviderBase();
+	
 	//~ Begin UObject Interface
 	virtual void BeginDestroy() override;
 	//~ End UObject Interface
@@ -72,7 +81,6 @@ public:
 	EVCamTargetViewportID GetTargetViewport() const { return TargetViewport; }
 	UFUNCTION(BlueprintCallable, Category = "Output")
 	void SetTargetViewport(EVCamTargetViewportID Value) { TargetViewport = Value; }
-	static FName GetTargetViewportPropertyName() { return GET_MEMBER_NAME_CHECKED(UVCamOutputProviderBase, TargetViewport); }
 	
 	UFUNCTION(BlueprintPure, Category = "Output")
 	TSubclassOf<UUserWidget> GetUMGClass() const { return UMGClass; }
@@ -112,6 +120,11 @@ public:
 	virtual void PostEditChangeProperty(FPropertyChangedEvent& PropertyChangedEvent) override;
 #endif
 
+
+	static FName GetIsActivePropertyName()			{ return GET_MEMBER_NAME_CHECKED(UVCamOutputProviderBase, bIsActive); }
+	static FName GetTargetViewportPropertyName()	{ return GET_MEMBER_NAME_CHECKED(UVCamOutputProviderBase, TargetViewport); }
+	static FName GetUMGClassPropertyName()			{ return GET_MEMBER_NAME_CHECKED(UVCamOutputProviderBase, UMGClass); }
+	
 protected:
 	
 	// If set, this output provider will execute every frame
diff --git a/Engine/Plugins/Experimental/VirtualCamera/Source/VCamCore/Public/UI/Switcher/VCamStateSwitcherWidget.h b/Engine/Plugins/Experimental/VirtualCamera/Source/VCamCore/Public/UI/Switcher/VCamStateSwitcherWidget.h
--- a/Engine/Plugins/Experimental/VirtualCamera/Source/VCamCore/Public/UI/Switcher/VCamStateSwitcherWidget.h
+++ b/Engine/Plugins/Experimental/VirtualCamera/Source/VCamCore/Public/UI/Switcher/VCamStateSwitcherWidget.h
@@ -39,6 +39,12 @@ public:
 	UFUNCTION(BlueprintPure, Category = "Connections")
 	FName GetCurrentState() const { return CurrentState; }
 
+	UFUNCTION(BlueprintPure, Category = "Connections")
+	TArray<FName> GetStates() const;
+
+	UFUNCTION(BlueprintPure, Category = "Connections")
+	bool GetStateInfo(FName State, FVCamWidgetConnectionState& OutStateInfo) const;
+
 #if WITH_EDITOR
 	virtual void PostEditChangeProperty(FPropertyChangedEvent& PropertyChangedEvent) override;
 #endif
diff --git a/Engine/Plugins/Experimental/VirtualCamera/Source/VCamCore/Public/UI/Switcher/VCamWidgetConnectionState.h b/Engine/Plugins/Experimental/VirtualCamera/Source/VCamCore/Public/UI/Switcher/VCamWidgetConnectionState.h
--- a/Engine/Plugins/Experimental/VirtualCamera/Source/VCamCore/Public/UI/Switcher/VCamWidgetConnectionState.h
+++ b/Engine/Plugins/Experimental/VirtualCamera/Source/VCamCore/Public/UI/Switcher/VCamWidgetConnectionState.h
@@ -8,12 +8,12 @@
 
 struct FVCamConnectionTargetSettings;
 
-USTRUCT()
+USTRUCT(BlueprintType)
 struct VCAMCORE_API FVCamWidgetConnectionState
 {
 	GENERATED_BODY()
 
 	/** A list of widgets to update */
-	UPROPERTY(EditAnywhere, Category = "Connection")
+	UPROPERTY(EditAnywhere, BlueprintReadWrite, Category = "Connection")
 	TArray<FWidgetConnectionConfig> WidgetConfigs;
 };
\ No newline at end of file
diff --git a/Engine/Plugins/Experimental/VirtualCamera/Source/VCamCore/Public/UI/Switcher/WidgetConnectionConfig.h b/Engine/Plugins/Experimental/VirtualCamera/Source/VCamCore/Public/UI/Switcher/WidgetConnectionConfig.h
--- a/Engine/Plugins/Experimental/VirtualCamera/Source/VCamCore/Public/UI/Switcher/WidgetConnectionConfig.h
+++ b/Engine/Plugins/Experimental/VirtualCamera/Source/VCamCore/Public/UI/Switcher/WidgetConnectionConfig.h
@@ -9,24 +9,23 @@
 class UVCamStateSwitcherWidget;
 class UVCamWidget;
 
-USTRUCT()
+USTRUCT(BlueprintType)
 struct VCAMCORE_API FWidgetConnectionConfig
 {
 	GENERATED_BODY()
 	
 	/**
 	 * Defines the widget whose connections should be updated. Must be a VCamWidget.
-	 * This is a reference to a property on the StateSwitcherVCamWidget subclass, intended to have been generated by checking
-	 * "Is Variable" in the Designer widget mode.
+	 * This is the name of a child widget and can be used as argument to UWidgetTree::FindWidget.
 	 */
-	UPROPERTY(EditAnywhere, Category = "Connection")
+	UPROPERTY(EditAnywhere, BlueprintReadWrite, Category = "Connection")
 	FName Widget = NAME_None;
 	
 	/**
 	 * Key: A valid key of in UVCamWidget::Connections of the widget identified via WidgetProperty.
 	 * Value: The settings to use for this connection.
 	 */
-	UPROPERTY(EditAnywhere, Category = "Connection")
+	UPROPERTY(EditAnywhere, BlueprintReadWrite, Category = "Connection")
 	TMap<FName, FVCamConnectionTargetSettings> ConnectionTargets;
 
 	UVCamWidget* ResolveWidget(UVCamStateSwitcherWidget* OwnerWidget) const;
diff --git a/Engine/Plugins/Experimental/VirtualCamera/Source/VCamCore/Public/UI/VCamWidget.h b/Engine/Plugins/Experimental/VirtualCamera/Source/VCamCore/Public/UI/VCamWidget.h
--- a/Engine/Plugins/Experimental/VirtualCamera/Source/VCamCore/Public/UI/VCamWidget.h
+++ b/Engine/Plugins/Experimental/VirtualCamera/Source/VCamCore/Public/UI/VCamWidget.h
@@ -30,6 +30,9 @@ class VCAMCORE_API UVCamWidget : public UUserWidget
 {
 	GENERATED_BODY()
 public:
+
+	UVCamWidget(const FObjectInitializer& ObjectInitializer);
+	
 	virtual void NativeDestruct() override;
 	
 	/*
@@ -60,6 +63,13 @@ public:
 	 */
 	UPROPERTY(EditAnywhere, BlueprintReadOnly, Category="VCam Input", meta=(EditCondition="bRegisterForInput"))
 	int32 InputContextPriority = 0;
+
+	DECLARE_MULTICAST_DELEGATE(FConnectionsReinitializedDelegate);
+	DECLARE_DYNAMIC_MULTICAST_DELEGATE(FConnectionsReinitializedDelegate_Blueprint);
+	FConnectionsReinitializedDelegate OnPostConnectionsReinitializedDelegate;
+	/** Called when ReinitializeConnections is called. */
+	UPROPERTY(BlueprintAssignable, meta = (DisplayName = "OnConnectionsReinitialized"))
+	FConnectionsReinitializedDelegate_Blueprint OnPostConnectionsReinitializedDelegate_Blueprint;
 	
 	/*
 	 * Event called when a specific connection has been updated
diff --git a/Engine/Plugins/Experimental/VirtualCamera/Source/VCamCoreEditor/Private/Customizations/ConnectionTargetSettingsTypeCustomization.cpp b/Engine/Plugins/Experimental/VirtualCamera/Source/VCamCoreEditor/Private/Customizations/ConnectionTargetSettingsTypeCustomization.cpp
--- a/Engine/Plugins/Experimental/VirtualCamera/Source/VCamCoreEditor/Private/Customizations/ConnectionTargetSettingsTypeCustomization.cpp
+++ b/Engine/Plugins/Experimental/VirtualCamera/Source/VCamCoreEditor/Private/Customizations/ConnectionTargetSettingsTypeCustomization.cpp
@@ -218,9 +218,6 @@ namespace UE::VCamCoreEditor::Private
 		TAttribute<bool> HasDataSourceAttr
 		) const
 	{
-		FLevelEditorModule& LevelEditorModule = FModuleManager::GetModuleChecked<FLevelEditorModule>(TEXT("LevelEditor"));
-		LevelEditorModule.OnActorSelectionChanged();
-		
 		Row.CustomWidget()
 			.NameContent()
 			[
@@ -292,9 +289,7 @@ namespace UE::VCamCoreEditor::Private
 									FUIAction(
 										FExecuteAction::CreateLambda([PropertyHandle, Name]()
 										{
-											PropertyHandle->NotifyPreChange();
 											PropertyHandle->SetValue(Name);
-											PropertyHandle->NotifyPostChange(EPropertyChangeType::ValueSet);
 										})),
 										NAME_None,
 										EUserInterfaceActionType::Button
diff --git a/Engine/Plugins/Experimental/VirtualCamera/Source/VCamCoreEditor/Private/Customizations/OutputProvider/ConnectionRemapCustomization_VCamWidget.cpp b/Engine/Plugins/Experimental/VirtualCamera/Source/VCamCoreEditor/Private/Customizations/OutputProvider/ConnectionRemapCustomization_VCamWidget.cpp
new file mode 100644
--- /dev/null
+++ b/Engine/Plugins/Experimental/VirtualCamera/Source/VCamCoreEditor/Private/Customizations/OutputProvider/ConnectionRemapCustomization_VCamWidget.cpp
@@ -0,0 +1,80 @@
+// Copyright Epic Games, Inc. All Rights Reserved.
+
+#include "ConnectionRemapCustomization_VCamWidget.h"
+
+#include "UI/VCamWidget.h"
+
+#include "Algo/AnyOf.h"
+#include "Customization/IConnectionRemapUtils.h"
+#include "DetailWidgetRow.h"
+#include "LogVCamEditor.h"
+
+#define LOCTEXT_NAMESPACE "FConnectionRemapCustomization_VCamWidget"
+
+namespace UE::VCamCoreEditor::Private
+{
+	TSharedRef<IConnectionRemapCustomization> FConnectionRemapCustomization_VCamWidget::Make()
+	{
+		return MakeShared<FConnectionRemapCustomization_VCamWidget>();
+	}
+
+	bool FConnectionRemapCustomization_VCamWidget::CanGenerateGroup(const FShouldGenerateArgs& Args) const
+	{
+		return Args.CustomizedWidget->Connections.Num() > 0
+			&& (!Args.DisplaySettings.bOnlyShowManuallyConfiguredConnections
+				|| Algo::AnyOf(Args.CustomizedWidget->Connections, [](const TPair<FName, FVCamConnection>& Pair)
+				{
+					return Pair.Value.bManuallyConfigureConnection;
+				}));
+	}
+
+	void FConnectionRemapCustomization_VCamWidget::Customize(const FConnectionRemapCustomizationArgs& Args)
+	{
+		Widget = Args.CustomizedWidget;
+		
+		for (const TPair<FName, FVCamConnection>& Connection : Args.CustomizedWidget->Connections)
+		{
+			if (Args.DisplaySettings.bOnlyShowManuallyConfiguredConnections && !Connection.Value.bManuallyConfigureConnection)
+			{
+				continue;
+			}
+			
+			// Overwrite bManuallyConfigureConnection because otherwise its EditCondition will hide the property!
+			FVCamConnection ConnectionCopy = Connection.Value;
+			ConnectionCopy.bManuallyConfigureConnection = true;
+			TStructOnScope<FVCamConnection> StructOnScope(ConnectionCopy);
+			
+			const FName ConnectionName = Connection.Key;
+			Args.Utils->AddConnection(FAddConnectionArgs{
+				Args.WidgetGroup,
+				ConnectionName,
+				GET_MEMBER_NAME_CHECKED(FVCamConnection, ConnectionTargetSettings),
+				// Object slicing occurs here
+				static_cast<FStructOnScope>(MoveTemp(StructOnScope)),
+				FOnTargetSettingsChanged::CreateSP(this, &FConnectionRemapCustomization_VCamWidget::OnTargetSettingsChanged, ConnectionName),
+				Args.Utils->GetRegularFont()
+			});
+		}
+	}
+
+	void FConnectionRemapCustomization_VCamWidget::OnTargetSettingsChanged(const FVCamConnectionTargetSettings& NewConnectionTargetSettings, FName ConnectionName) const
+	{
+		if (!NewConnectionTargetSettings.HasValidSettings())
+		{
+			return;
+		}
+		
+		EConnectionUpdateResult Result = EConnectionUpdateResult::NoConnectionsUpdated;
+		
+		if (Widget.IsValid())
+		{
+			const TMap<FName, FVCamConnectionTargetSettings> NewConnectionTargets {{ ConnectionName, NewConnectionTargetSettings }};
+			constexpr bool bReinitializeOnSuccess = true;
+			Widget->UpdateConnectionTargets(NewConnectionTargets, bReinitializeOnSuccess, Result);
+		}
+		
+		UE_CLOG(Result == EConnectionUpdateResult::NoConnectionsUpdated, LogVCamEditor, Warning, TEXT("Failed to update connection %s"), *ConnectionName.ToString());
+	}
+}
+
+#undef LOCTEXT_NAMESPACE
\ No newline at end of file
diff --git a/Engine/Plugins/Experimental/VirtualCamera/Source/VCamCoreEditor/Private/Customizations/OutputProvider/ConnectionRemapCustomization_VCamWidget.h b/Engine/Plugins/Experimental/VirtualCamera/Source/VCamCoreEditor/Private/Customizations/OutputProvider/ConnectionRemapCustomization_VCamWidget.h
new file mode 100644
--- /dev/null
+++ b/Engine/Plugins/Experimental/VirtualCamera/Source/VCamCoreEditor/Private/Customizations/OutputProvider/ConnectionRemapCustomization_VCamWidget.h
@@ -0,0 +1,29 @@
+// Copyright Epic Games, Inc. All Rights Reserved.
+
+#pragma once
+
+#include "CoreMinimal.h"
+#include "Customization/IConnectionRemapCustomization.h"
+#include "UI/VCamConnectionStructs.h"
+
+namespace UE::VCamCoreEditor::Private
+{
+	class FConnectionRemapCustomization_VCamWidget : public IConnectionRemapCustomization
+	{
+	public:
+
+		static TSharedRef<IConnectionRemapCustomization> Make();
+
+		//~ Begin IConnectionRemappingCustomization Interface
+		virtual bool CanGenerateGroup(const FShouldGenerateArgs& Args) const override;
+		virtual void Customize(const FConnectionRemapCustomizationArgs& Args) override;
+		//~ End IConnectionRemappingCustomization Interface
+
+	private:
+
+		TWeakObjectPtr<UVCamWidget> Widget;
+		
+		void OnTargetSettingsChanged(const FVCamConnectionTargetSettings& NewConnectionTargetSettings, FName ConnectionName) const;
+	};
+}
+
diff --git a/Engine/Plugins/Experimental/VirtualCamera/Source/VCamCoreEditor/Private/Customizations/OutputProvider/ConnectionRemapUtilsImpl.cpp b/Engine/Plugins/Experimental/VirtualCamera/Source/VCamCoreEditor/Private/Customizations/OutputProvider/ConnectionRemapUtilsImpl.cpp
new file mode 100644
--- /dev/null
+++ b/Engine/Plugins/Experimental/VirtualCamera/Source/VCamCoreEditor/Private/Customizations/OutputProvider/ConnectionRemapUtilsImpl.cpp
@@ -0,0 +1,82 @@
+// Copyright Epic Games, Inc. All Rights Reserved.
+
+#include "ConnectionRemapUtilsImpl.h"
+
+#include "DetailWidgetRow.h"
+#include "IDetailGroup.h"
+#include "UI/VCamConnectionStructs.h"
+#include "Widgets/Text/STextBlock.h"
+
+namespace UE::VCamCoreEditor::Private
+{
+	FConnectionRemapUtilsImpl::FConnectionRemapUtilsImpl(TSharedRef<IDetailLayoutBuilder> Builder)
+		: Builder(MoveTemp(Builder))
+	{}
+
+	void FConnectionRemapUtilsImpl::AddConnection(FAddConnectionArgs Args)
+	{
+		TSharedPtr<IDetailLayoutBuilder> BuilderPin =  Builder.Pin();
+		if (!BuilderPin)
+		{
+			return;
+		}
+		
+		TSharedPtr<FStructOnScope> StructData;
+		if (const TSharedPtr<FStructOnScope>* ExistingStructData = AddedConnections.Find(Args.PropertyName))
+		{
+			StructData = *ExistingStructData;
+		}
+		else
+		{
+			StructData = MakeShared<FStructOnScope>(MoveTemp(Args.StructData));
+			AddedConnections.Add(Args.PropertyName, StructData);
+		}
+		
+		const TSharedPtr<IPropertyHandle> PropertyHandle = BuilderPin->AddStructurePropertyData(StructData, Args.PropertyName);
+		if (!ensureMsgf(PropertyHandle, TEXT("No property of %s found on struct"), *Args.PropertyName.ToString()))
+		{
+			return;
+		}
+		
+		PropertyHandle->SetOnChildPropertyValueChanged(FSimpleDelegate::CreateLambda([PropertyHandle, Callback = MoveTemp(Args.OnTargetSettingsChangedDelegate)]()
+		{
+			void* Data;
+			if (ensure(PropertyHandle->GetValueData(Data) == FPropertyAccess::Success))
+			{
+				Callback.Execute(*(FVCamConnectionTargetSettings*)Data);
+			}
+		}));
+
+		Args.DetailGroup.AddPropertyRow(PropertyHandle.ToSharedRef())
+			.DisplayName(FText::FromName(Args.ConnectionName)) 
+			.CustomWidget(true)
+			.NameContent()
+			[
+				SNew(STextBlock)
+				.Font(Args.Font)
+				.Text(FText::FromName(Args.ConnectionName))
+			]
+			.ValueContent()
+			[
+				PropertyHandle->CreatePropertyValueWidget()
+			];
+	}
+
+	FSlateFontInfo FConnectionRemapUtilsImpl::GetRegularFont() const
+	{
+		if (TSharedPtr<IDetailLayoutBuilder> BuilderPin = Builder.Pin())
+		{
+			return BuilderPin->GetDetailFont();
+		}
+		return FSlateFontInfo{};
+	}
+
+	void FConnectionRemapUtilsImpl::ForceRefreshProperties() const
+	{
+		if (TSharedPtr<IDetailLayoutBuilder> BuilderPin = Builder.Pin())
+		{
+			BuilderPin->ForceRefreshDetails();
+		}
+	}
+}
+
diff --git a/Engine/Plugins/Experimental/VirtualCamera/Source/VCamCoreEditor/Private/Customizations/OutputProvider/ConnectionRemapUtilsImpl.h b/Engine/Plugins/Experimental/VirtualCamera/Source/VCamCoreEditor/Private/Customizations/OutputProvider/ConnectionRemapUtilsImpl.h
new file mode 100644
--- /dev/null
+++ b/Engine/Plugins/Experimental/VirtualCamera/Source/VCamCoreEditor/Private/Customizations/OutputProvider/ConnectionRemapUtilsImpl.h
@@ -0,0 +1,30 @@
+// Copyright Epic Games, Inc. All Rights Reserved.
+
+#pragma once
+
+#include "CoreMinimal.h"
+#include "DetailLayoutBuilder.h"
+#include "Customization/IConnectionRemapUtils.h"
+
+namespace UE::VCamCoreEditor::Private
+{
+	class FConnectionRemapUtilsImpl : public IConnectionRemapUtils
+	{
+	public:
+		
+		FConnectionRemapUtilsImpl(TSharedRef<IDetailLayoutBuilder> Builder);
+
+		//~ Begin IConnectionRemapUtils Interface
+		virtual void AddConnection(FAddConnectionArgs Args) override;
+		virtual FSlateFontInfo GetRegularFont() const override;
+		virtual void ForceRefreshProperties() const override;
+		//~ End IConnectionRemapUtils Interface
+
+	private:
+		
+		TWeakPtr<IDetailLayoutBuilder> Builder;
+		TMap<FName, TSharedPtr<FStructOnScope>> AddedConnections;
+	};
+
+}
+
diff --git a/Engine/Plugins/Experimental/VirtualCamera/Source/VCamCoreEditor/Private/Customizations/OutputProvider/OutputProviderLayoutCustomization.cpp b/Engine/Plugins/Experimental/VirtualCamera/Source/VCamCoreEditor/Private/Customizations/OutputProvider/OutputProviderLayoutCustomization.cpp
new file mode 100644
--- /dev/null
+++ b/Engine/Plugins/Experimental/VirtualCamera/Source/VCamCoreEditor/Private/Customizations/OutputProvider/OutputProviderLayoutCustomization.cpp
@@ -0,0 +1,399 @@
+// Copyright Epic Games, Inc. All Rights Reserved.
+
+#include "OutputProviderLayoutCustomization.h"
+
+#include "DetailCategoryBuilder.h"
+#include "DetailLayoutBuilder.h"
+#include "Output/VCamOutputProviderBase.h"
+#include "UI/VCamWidget.h"
+#include "VCamCoreEditorModule.h"
+#include "VPFullScreenUserWidget.h"
+
+#include "Blueprint/WidgetTree.h"
+#include "DetailWidgetRow.h"
+#include "Editor.h"
+#include "Framework/MultiBox/MultiBoxBuilder.h"
+#include "IDetailGroup.h"
+#include "PropertyCustomizationHelpers.h"
+#include "SSimpleComboButton.h"
+#include "Widgets/Images/SImage.h"
+#include "Widgets/Input/SButton.h"
+#include "Widgets/Layout/SBox.h"
+#include "Widgets/SBoxPanel.h"
+#include "Widgets/Text/STextBlock.h"
+
+#define LOCTEXT_NAMESPACE "FOutputProviderCustomization"
+
+namespace UE::VCamCoreEditor::Private
+{
+	namespace Private
+	{
+		/** Static because IDetailCustomization is destroyed when details panel is refreshed */
+		static FTargetConnectionDisplaySettings DisplaySettings;
+		
+		class SPropertyEditorButton : public SCompoundWidget
+		{
+		public:
+
+			SLATE_BEGIN_ARGS( SPropertyEditorButton ) 
+				: _Text( )
+				, _Image( FAppStyle::GetBrush("Default") )
+			{}
+				SLATE_ATTRIBUTE( FText, Text )
+				SLATE_ARGUMENT(const FSlateBrush*, Image )
+				SLATE_EVENT(FSimpleDelegate, OnClickAction)
+			SLATE_END_ARGS()
+
+			void Construct( const FArguments& InArgs )
+			{
+				OnClickAction = InArgs._OnClickAction;
+
+				ChildSlot
+				[
+					SNew(SBox)
+					.HAlign(HAlign_Center)
+					.VAlign(VAlign_Center)
+					.WidthOverride(22)
+					.HeightOverride(22)
+					.ToolTipText(InArgs._Text)
+					[
+						SNew(SButton)
+						.ButtonStyle( FAppStyle::Get(), "SimpleButton" )
+						.OnClicked( this, &SPropertyEditorButton::OnClick )
+						.ContentPadding(0)
+						[ 
+							SNew( SImage )
+							.Image( InArgs._Image )
+							.ColorAndOpacity( FSlateColor::UseForeground() )
+						]
+					]
+				]; 
+			}
+			
+		private:
+			FSimpleDelegate OnClickAction;
+			
+			FReply OnClick()
+			{
+				OnClickAction.ExecuteIfBound();
+				return FReply::Handled();
+			}
+		};
+		
+		struct FWidgetDisplayInfo
+		{
+			FName UniqueName;
+			FText WidgetDisplayName;
+			bool bNeedsToDisplayParentWidgetName;
+		};
+		
+		static TMap<TWeakObjectPtr<UVCamWidget>, FWidgetDisplayInfo> GenerateWidgetRowNames(TArray<TWeakObjectPtr<UVCamWidget>> SortedWidgets)
+		{
+			TMap<FName, int32> NameCounter;
+			TMap<TWeakObjectPtr<UVCamWidget>, FWidgetDisplayInfo> TargetDisplayInfo;
+			for (const TWeakObjectPtr<UVCamWidget>& Widget : SortedWidgets)
+			{
+				FName UniqueNameForWidget = Widget->GetFName();
+				int32& NameCount = NameCounter.FindOrAdd(UniqueNameForWidget);
+				++NameCount;
+				// If widgets have the same names (can happen with multiple Blueprints) fall back to expensive string construction
+				if (NameCount > 1)
+				{
+					UniqueNameForWidget = FName(FString::Printf(TEXT("%s_%d"), *UniqueNameForWidget.ToString(), NameCount));
+				}
+			
+				FWidgetDisplayInfo& Info = TargetDisplayInfo.Add(Widget);
+				Info.UniqueName = UniqueNameForWidget;
+				Info.WidgetDisplayName = FText::FromName(Widget->GetFName());
+			}
+
+			for (TPair<TWeakObjectPtr<UVCamWidget>, FWidgetDisplayInfo>& Pair : TargetDisplayInfo)
+			{
+				Pair.Value.bNeedsToDisplayParentWidgetName = NameCounter[Pair.Key->GetFName()] > 1;
+			}
+			return TargetDisplayInfo;
+		}
+
+		static UObject* GetBlueprintFrom(TWeakObjectPtr<UVCamWidget> WidgetToBrowseTo)
+		{
+			if (WidgetToBrowseTo.IsValid())
+			{
+				TArray<FAssetData> AssetDatas;
+				UWidgetTree* WidgetTree = Cast<UWidgetTree>(WidgetToBrowseTo->GetOuter());
+				UUserWidget* OwningWidget = WidgetTree ? Cast<UUserWidget>(WidgetTree->GetOuter()) : nullptr;
+				UObject* Blueprint = OwningWidget ? OwningWidget->GetClass()->ClassGeneratedBy : nullptr;
+				return Blueprint;
+			}
+			return nullptr;
+		}
+
+		static void BrowseToWidgetSource(TWeakObjectPtr<UVCamWidget> WidgetToBrowseTo)
+		{
+			if (UObject* Blueprint = GetBlueprintFrom(WidgetToBrowseTo))
+			{
+				TArray<UObject*> ObjectsToFocus { Blueprint };
+				GEditor->SyncBrowserToObjects(ObjectsToFocus);
+			}
+		}
+
+		static void OpenAssetAndSelectWidget(TWeakObjectPtr<UVCamWidget> WidgetToSelect)
+		{
+			if (UObject* Blueprint = GetBlueprintFrom(WidgetToSelect))
+			{
+				GEditor->EditObject(Blueprint);
+			}
+		}
+	}
+	
+	TSharedRef<IDetailCustomization> FOutputProviderLayoutCustomization::MakeInstance()
+	{
+		return MakeShared<FOutputProviderLayoutCustomization>();
+	}
+
+	FOutputProviderLayoutCustomization::~FOutputProviderLayoutCustomization()
+	{
+		// Technically unsubscribing is not needed because delegates clean dangling reference automatically, but let's not dangle on purpose...
+		if (CustomizedOutputProvider.IsValid())
+		{
+			CustomizedOutputProvider->OnActivatedDelegate.Remove(OnActivatedDelegateHandle);
+		}
+		
+		ClearWidgetData(EditableWidgets);
+	}
+
+	void FOutputProviderLayoutCustomization::CustomizeDetails(IDetailLayoutBuilder& DetailBuilder)
+	{
+		TArray<TWeakObjectPtr<UObject>> CustomizedObjects;
+		DetailBuilder.GetObjectsBeingCustomized(CustomizedObjects);
+		CustomizedOutputProvider = CustomizedObjects.Num() == 1
+			? Cast<UVCamOutputProviderBase>(CustomizedObjects[0])
+			: nullptr;
+		if (!CustomizedOutputProvider.IsValid())
+		{
+			return;
+		}
+		
+		if (!OnActivatedDelegateHandle.IsValid())
+		{
+			OnActivatedDelegateHandle = CustomizedOutputProvider->OnActivatedDelegate.AddSP(this, &FOutputProviderLayoutCustomization::OnActivationChanged);
+		}
+
+		// Important properties should show before widgets, then ...
+		IDetailCategoryBuilder& FirstCategory = DetailBuilder.EditCategory(TEXT("Primary"));
+		FirstCategory.SetSortOrder(0);
+		FirstCategory.AddProperty(DetailBuilder.GetProperty(UVCamOutputProviderBase::GetIsActivePropertyName()));
+		FirstCategory.AddProperty(DetailBuilder.GetProperty(UVCamOutputProviderBase::GetTargetViewportPropertyName()));
+		FirstCategory.AddProperty(DetailBuilder.GetProperty(UVCamOutputProviderBase::GetUMGClassPropertyName()));
+
+		// ... the widgets should show after important properties, and ...
+		RebuildWidgetData();
+		if (!EditableWidgets.IsEmpty())
+		{
+			IDetailCategoryBuilder& CategoryBuilder = DetailBuilder.EditCategory(TEXT("Widgets"), LOCTEXT("Widget", "Widgets"));
+			CategoryBuilder.SetSortOrder(1);
+			IDetailGroup& WidgetGroup = CategoryBuilder.AddGroup(TEXT("Widgets"), LOCTEXT("WidgetsLabel", "Widgets"));
+			ExtendWidgetsRow(DetailBuilder, WidgetGroup);
+			GenerateWidgetRows(WidgetGroup, DetailBuilder);
+		}
+
+		// ... all other properties should be shown after widgets
+	}
+
+	void FOutputProviderLayoutCustomization::CustomizeDetails(const TSharedPtr<IDetailLayoutBuilder>& DetailBuilder)
+	{
+		if (DetailBuilder != WeakDetailBuilder)
+		{
+			WeakDetailBuilder = DetailBuilder;
+		}
+		
+		IDetailCustomization::CustomizeDetails(DetailBuilder);
+	}
+	
+	FDetailWidgetRow FOutputProviderLayoutCustomization::ExtendWidgetsRow(IDetailLayoutBuilder& DetailBuilder,IDetailGroup& WidgetGroup)
+	{
+		return WidgetGroup.HeaderRow()
+			.NameContent()
+			[
+				SNew(STextBlock)
+					.Font(DetailBuilder.GetDetailFont())
+					.Text(LOCTEXT("Widget", "Widgets"))
+			]
+			.ValueContent()
+			[
+				SNew(SSimpleComboButton)
+				.Icon(FAppStyle::Get().GetBrush("DetailsView.ViewOptions"))
+				.OnGetMenuContent_Lambda([this]()
+				{
+					FMenuBuilder MenuBuilder(true, nullptr);
+					MenuBuilder.AddMenuEntry(
+						LOCTEXT("FTargetConnectionDisplaySettings.bOnlyShowManuallyConfiguredConnections", "Only Manually Configured Connections"),
+						FText::GetEmpty(),
+						FSlateIcon(),
+						FUIAction(
+							FExecuteAction::CreateLambda([this]()
+							{
+								Private::DisplaySettings.bOnlyShowManuallyConfiguredConnections = !Private::DisplaySettings.bOnlyShowManuallyConfiguredConnections;
+								ForceRefreshDetailsIfSafe();
+							}),
+							FCanExecuteAction::CreateLambda([](){ return true; }),
+							FIsActionChecked::CreateLambda([](){ return Private::DisplaySettings.bOnlyShowManuallyConfiguredConnections; })
+							),
+							NAME_None,
+							EUserInterfaceActionType::Check
+						);
+					return MenuBuilder.MakeWidget();
+				})
+			];
+	}
+
+	void FOutputProviderLayoutCustomization::RebuildWidgetData()
+	{
+		const UVPFullScreenUserWidget* FullScreenUserWidget = CustomizedOutputProvider->GetUMGWidget();
+		const UUserWidget* RootWidget = FullScreenUserWidget ? FullScreenUserWidget->GetWidget() : nullptr;
+		UWidgetTree* WidgetTree = RootWidget ? RootWidget->WidgetTree : nullptr;
+		if (!WidgetTree)
+		{
+			return;
+		}
+
+		// Move so EditableWidgets gets reset and does not retain old references
+		TMap<TWeakObjectPtr<UVCamWidget>, FWidgetData> OldEditableWidgets = MoveTemp(EditableWidgets);
+		
+		TQueue<UWidgetTree*> SearchQueue;
+		SearchQueue.Enqueue(WidgetTree);
+		UWidgetTree* CurrentTree;
+		while (SearchQueue.Dequeue(CurrentTree))
+		{
+			// We recurse ourselves because the other functions skip UUserWidgets
+			CurrentTree->ForEachWidget([this, &OldEditableWidgets, &SearchQueue](UWidget* Widget)
+			{
+				if (UVCamWidget* VCamWidget = Cast<UVCamWidget>(Widget))
+				{
+					FWidgetData ExistingWidgetData;
+					if (OldEditableWidgets.RemoveAndCopyValue(VCamWidget, ExistingWidgetData))
+					{
+						EditableWidgets.Emplace(VCamWidget, MoveTemp(ExistingWidgetData));
+					}
+					else if (TSharedPtr<IConnectionRemapCustomization> Customization = FVCamCoreEditorModule::Get().CreateConnectionRemapCustomization(VCamWidget->GetClass()))
+					{
+						const TWeakObjectPtr<UVCamWidget> WeakVCamWidget(VCamWidget);
+						FDelegateHandle OnPostConnectionsReinitializedDelegateHandle = VCamWidget->OnPostConnectionsReinitializedDelegate.AddSP(
+							this, &FOutputProviderLayoutCustomization::OnPostConnectionsReinitialized
+							);
+						EditableWidgets.Emplace(
+							WeakVCamWidget,
+							FWidgetData{ OnPostConnectionsReinitializedDelegateHandle, Customization, MakeShared<FConnectionRemapUtilsImpl>(WeakDetailBuilder.Pin().ToSharedRef()) }
+							);
+					}
+
+					if (VCamWidget->WidgetTree)
+					{
+						SearchQueue.Enqueue(VCamWidget->WidgetTree);
+					}
+				}
+			});
+		}
+		
+		// The hierarchy may have changed so anything that is left is not part of the hierarchy and can be unsubscribed from
+		ClearWidgetData(OldEditableWidgets);
+	}
+
+	void FOutputProviderLayoutCustomization::GenerateWidgetRows(IDetailGroup& RootWidgetGroup, IDetailLayoutBuilder& DetailBuilder)
+	{
+		TArray<TWeakObjectPtr<UVCamWidget>> SortedWidgets;
+		EditableWidgets.GenerateKeyArray(SortedWidgets);
+		SortedWidgets.Sort([](const TWeakObjectPtr<UVCamWidget>& First, const TWeakObjectPtr<UVCamWidget>& Second)
+		{
+			return First->GetName() < Second->GetName(); 
+		});
+		
+		TMap<TWeakObjectPtr<UVCamWidget>, Private::FWidgetDisplayInfo> WidgetDisplayData = Private::GenerateWidgetRowNames(SortedWidgets);
+		for (const TWeakObjectPtr<UVCamWidget>& Widget : SortedWidgets)
+		{
+			const FWidgetData* WidgetData = EditableWidgets.Find(Widget);
+			if (!ensure(WidgetData))
+			{
+				continue;
+			}
+
+			if (!WidgetData->Customization->CanGenerateGroup({ Widget, Private::DisplaySettings}))
+			{
+				continue;
+			}
+			
+			const Private::FWidgetDisplayInfo& DisplayInfo = WidgetDisplayData[Widget];
+			// Row name is "WidgetName" or "WidgetName (Outer name)" 
+			const FText RowDisplayName = DisplayInfo.bNeedsToDisplayParentWidgetName
+				? FText::Format(LOCTEXT("WidgetNameFmt", "{0} ({1})"), DisplayInfo.WidgetDisplayName, FText::FromName(Widget->GetOuter()->GetFName()))
+				: DisplayInfo.WidgetDisplayName;
+			IDetailGroup& WidgetGroup = RootWidgetGroup.AddGroup(DisplayInfo.UniqueName, RowDisplayName);
+			WidgetGroup.HeaderRow()
+				.NameContent()
+				[
+					SNew(STextBlock)
+					.Text(RowDisplayName)
+					.Font(DetailBuilder.GetDetailFont())
+				]
+				.ValueContent()
+				[
+					CreateControlWidgets(Widget)
+				];
+			WidgetData->Customization->Customize({ WidgetGroup, WidgetData->RemapUtils.ToSharedRef(), Widget, Private::DisplaySettings });
+		}
+	}
+	
+	TSharedRef<SHorizontalBox> FOutputProviderLayoutCustomization::CreateControlWidgets(const TWeakObjectPtr<UVCamWidget>& Widget) const
+	{
+		return SNew(SHorizontalBox)
+			+ SHorizontalBox::Slot()
+			.AutoWidth()
+			[
+				PropertyCustomizationHelpers::MakeBrowseButton(
+					FSimpleDelegate::CreateStatic(&Private::BrowseToWidgetSource, Widget),
+					LOCTEXT("Browse", "Browse to the Blueprint which owns this widget"),
+					TAttribute<bool>::CreateLambda([Widget]() { return Widget.IsValid(); })
+				)
+
+			]
+			+ SHorizontalBox::Slot()
+			.AutoWidth()
+			[
+				SNew(Private::SPropertyEditorButton)
+					.Image(FAppStyle::Get().GetBrush(TEXT("Icons.Blueprint")))
+					.Text(LOCTEXT("OpenBlueprint", "Open the Blueprint owning this widget and select it in the hierarchy"))
+					.OnClickAction(FSimpleDelegate::CreateStatic(&Private::OpenAssetAndSelectWidget, Widget))
+			];
+	}
+	
+	void FOutputProviderLayoutCustomization::OnActivationChanged(bool bNewIsActivated) const
+	{
+		ForceRefreshDetailsIfSafe();
+	}
+	
+	void FOutputProviderLayoutCustomization::OnPostConnectionsReinitialized() const
+	{
+		ForceRefreshDetailsIfSafe();
+	}
+
+	void FOutputProviderLayoutCustomization::ForceRefreshDetailsIfSafe() const
+	{
+		if (TSharedPtr<IDetailLayoutBuilder> DetailBuilder = WeakDetailBuilder.Pin())
+		{
+			DetailBuilder->ForceRefreshDetails();
+		}
+	}
+
+	void FOutputProviderLayoutCustomization::ClearWidgetData(TMap<TWeakObjectPtr<UVCamWidget>, FWidgetData>& InEditableWidgets)
+	{
+		for (TPair<TWeakObjectPtr<UVCamWidget>, FWidgetData>& WidgetData : InEditableWidgets)
+		{
+			if (WidgetData.Key.IsValid())
+			{
+				WidgetData.Key->OnPostConnectionsReinitializedDelegate.Remove(WidgetData.Value.OnPostConnectionsReinitializedDelegateHandle);
+			}
+		}
+		InEditableWidgets.Empty();
+	}
+}
+
+#undef LOCTEXT_NAMESPACE
\ No newline at end of file
diff --git a/Engine/Plugins/Experimental/VirtualCamera/Source/VCamCoreEditor/Private/Customizations/OutputProvider/OutputProviderLayoutCustomization.h b/Engine/Plugins/Experimental/VirtualCamera/Source/VCamCoreEditor/Private/Customizations/OutputProvider/OutputProviderLayoutCustomization.h
new file mode 100644
--- /dev/null
+++ b/Engine/Plugins/Experimental/VirtualCamera/Source/VCamCoreEditor/Private/Customizations/OutputProvider/OutputProviderLayoutCustomization.h
@@ -0,0 +1,68 @@
+// Copyright Epic Games, Inc. All Rights Reserved.
+
+#pragma once
+
+#include "CoreMinimal.h"
+#include "ConnectionRemapUtilsImpl.h"
+#include "IDetailCustomization.h"
+#include "Customization/IConnectionRemapCustomization.h"
+#include "UObject/WeakObjectPtr.h"
+#include "UObject/WeakObjectPtrTemplates.h"
+
+class IDetailGroup;
+class IPropertyUtilities;
+class UVCamWidget;
+class UVCamOutputProviderBase;
+
+namespace UE::VCamCoreEditor::Private
+{
+	/** Exposes all VCam widgets' connection settings for as long as the output provider has a valid widget. */
+	class FOutputProviderLayoutCustomization
+		: public IDetailCustomization
+	{
+	public:
+
+		static TSharedRef<IDetailCustomization> MakeInstance();
+		virtual ~FOutputProviderLayoutCustomization() override;
+
+		//~ Begin IDetailCustomization Interface
+		virtual void CustomizeDetails(IDetailLayoutBuilder& DetailBuilder) override;
+		virtual void CustomizeDetails(const TSharedPtr<IDetailLayoutBuilder>& DetailBuilder) override;
+		//~ End IDetailCustomization Interface
+
+	private:
+
+		/** Exactly 1 object can be customized at a time. */
+		TWeakObjectPtr<UVCamOutputProviderBase> CustomizedOutputProvider;
+
+		TWeakPtr<IDetailLayoutBuilder> WeakDetailBuilder;
+		
+		/** Handle to UVCamOutputProviderBase::OnActivatedDelegate */
+		FDelegateHandle OnActivatedDelegateHandle;
+		
+		struct FWidgetData
+		{
+			/** Handle to UVCamWidget::OnPostConnectionsReinitialized */
+			FDelegateHandle OnPostConnectionsReinitializedDelegateHandle;
+
+			TSharedPtr<IConnectionRemapCustomization> Customization;
+			
+			/** Utils for IConnectionRemapCustomizations */
+			TSharedPtr<FConnectionRemapUtilsImpl> RemapUtils;
+		};
+		/** The widgets added to the details panel */
+		TMap<TWeakObjectPtr<UVCamWidget>, FWidgetData> EditableWidgets;
+		
+		FDetailWidgetRow ExtendWidgetsRow(IDetailLayoutBuilder& DetailBuilder, IDetailGroup& WidgetGroup);
+		void RebuildWidgetData();
+		void GenerateWidgetRows(IDetailGroup& RootWidgetGroup, IDetailLayoutBuilder& DetailBuilder);
+		TSharedRef<SHorizontalBox> CreateControlWidgets(const TWeakObjectPtr<UVCamWidget>& Widget) const;
+		
+		void OnActivationChanged(bool bNewIsActivated) const;
+		void OnPostConnectionsReinitialized() const;
+		void ForceRefreshDetailsIfSafe() const;
+		
+		static void ClearWidgetData(TMap<TWeakObjectPtr<UVCamWidget>, FWidgetData>& InEditableWidgets);
+	};
+}
+
diff --git a/Engine/Plugins/Experimental/VirtualCamera/Source/VCamCoreEditor/Private/Customizations/OutputProvider/OutputProviderTypeCustomization.cpp b/Engine/Plugins/Experimental/VirtualCamera/Source/VCamCoreEditor/Private/Customizations/OutputProvider/OutputProviderTypeCustomization.cpp
new file mode 100644
--- /dev/null
+++ b/Engine/Plugins/Experimental/VirtualCamera/Source/VCamCoreEditor/Private/Customizations/OutputProvider/OutputProviderTypeCustomization.cpp
@@ -0,0 +1,62 @@
+// Copyright Epic Games, Inc. All Rights Reserved.
+
+#include "OutputProviderTypeCustomization.h"
+
+#include "DetailCategoryBuilder.h"
+#include "Output/VCamOutputProviderBase.h"
+#include "UI/VCamWidget.h"
+
+#include "DetailWidgetRow.h"
+#include "IDetailChildrenBuilder.h"
+#include "PropertyCustomizationHelpers.h"
+
+#define LOCTEXT_NAMESPACE "FOutputProviderCustomization"
+
+namespace UE::VCamCoreEditor::Private
+{
+	TSharedRef<IPropertyTypeCustomization> FOutputProviderTypeCustomization::MakeInstance()
+	{
+		return MakeShared<FOutputProviderTypeCustomization>();
+	}
+
+	void FOutputProviderTypeCustomization::CustomizeHeader(TSharedRef<IPropertyHandle> PropertyHandle, FDetailWidgetRow& HeaderRow, IPropertyTypeCustomizationUtils& CustomizationUtils)
+	{}
+
+	void FOutputProviderTypeCustomization::CustomizeChildren(TSharedRef<IPropertyHandle> PropertyHandle, IDetailChildrenBuilder& ChildBuilder, IPropertyTypeCustomizationUtils& CustomizationUtils)
+	{
+		UObject* CustomizedObject;
+		const bool bSuccessGettingValue = PropertyHandle->GetValue(CustomizedObject) == FPropertyAccess::Success;
+		UVCamOutputProviderBase* CustomizedOutputProvider = bSuccessGettingValue
+			? Cast<UVCamOutputProviderBase>(CustomizedObject)
+			: nullptr;
+		if (!CustomizedOutputProvider)
+		{
+			return;
+		}
+		
+		uint32 NumChildren;
+		if (PropertyHandle->GetNumChildren(NumChildren) != FPropertyAccess::Success
+			|| NumChildren != 1)
+		{
+			return;
+		}
+		
+		IDetailPropertyRow* DetailRow = ChildBuilder.AddExternalObjects({ CustomizedOutputProvider },
+			FAddPropertyParams()
+				.CreateCategoryNodes(false)		// Avoid creating intermediate group expansion
+				.AllowChildren(true)			// Child properties should be shown
+				.HideRootObjectNode(false)		// Needed so we can use NameContent() & ValueContent() below
+				);
+		DetailRow->CustomWidget(true)
+			.NameContent()
+			[
+				PropertyHandle->CreatePropertyNameWidget()
+			]
+			.ValueContent()
+			[
+				PropertyHandle->CreatePropertyValueWidget()
+			];
+	}
+}
+
+#undef LOCTEXT_NAMESPACE
\ No newline at end of file
diff --git a/Engine/Plugins/Experimental/VirtualCamera/Source/VCamCoreEditor/Private/Customizations/OutputProvider/OutputProviderTypeCustomization.h b/Engine/Plugins/Experimental/VirtualCamera/Source/VCamCoreEditor/Private/Customizations/OutputProvider/OutputProviderTypeCustomization.h
new file mode 100644
--- /dev/null
+++ b/Engine/Plugins/Experimental/VirtualCamera/Source/VCamCoreEditor/Private/Customizations/OutputProvider/OutputProviderTypeCustomization.h
@@ -0,0 +1,28 @@
+// Copyright Epic Games, Inc. All Rights Reserved.
+
+#pragma once
+
+#include "CoreMinimal.h"
+#include "IPropertyTypeCustomization.h"
+
+class IDetailGroup;
+class IPropertyUtilities;
+class UVCamWidget;
+class UVCamOutputProviderBase;
+
+namespace UE::VCamCoreEditor::Private
+{
+	/** Needed to customize display UPROPERTY(Instanced) instances. Causes FOutputProviderLayoutCustomization to be used. */
+	class FOutputProviderTypeCustomization
+		: public IPropertyTypeCustomization
+	{
+	public:
+
+		static TSharedRef<IPropertyTypeCustomization> MakeInstance();
+
+		//~ Begin IDetailCustomization Interface
+		virtual void CustomizeHeader(TSharedRef<IPropertyHandle> PropertyHandle, FDetailWidgetRow& HeaderRow, IPropertyTypeCustomizationUtils& CustomizationUtils) override;
+		virtual void CustomizeChildren(TSharedRef<IPropertyHandle> PropertyHandle, IDetailChildrenBuilder& ChildBuilder, IPropertyTypeCustomizationUtils& CustomizationUtils) override;
+		//~ End IDetailCustomization Interface
+	};
+}
\ No newline at end of file
diff --git a/Engine/Plugins/Experimental/VirtualCamera/Source/VCamCoreEditor/Private/VCamCoreEditorModule.cpp b/Engine/Plugins/Experimental/VirtualCamera/Source/VCamCoreEditor/Private/VCamCoreEditorModule.cpp
--- a/Engine/Plugins/Experimental/VirtualCamera/Source/VCamCoreEditor/Private/VCamCoreEditorModule.cpp
+++ b/Engine/Plugins/Experimental/VirtualCamera/Source/VCamCoreEditor/Private/VCamCoreEditorModule.cpp
@@ -1,68 +1,160 @@
 // Copyright Epic Games, Inc. All Rights Reserved.
 
-#include "Modules/ModuleInterface.h"
-#include "Modules/ModuleManager.h"
+#include "VCamCoreEditorModule.h"
 
 #include "Customizations/ConnectionTargetSettingsTypeCustomization.h"
+#include "Customizations/OutputProvider/OutputProviderLayoutCustomization.h"
 #include "Customizations/VCamBaseActorCustomization.h"
 #include "Customizations/VCamInputProfileCustomization.h"
 #include "Customizations/StateSwitcher/WidgetConnectionConfigTypeCustomization.h"
 #include "UI/Switcher/WidgetConnectionConfig.h"
 #include "VCamBaseActor.h"
 #include "VCamComponent.h"
+#include "Output/VCamOutputProviderBase.h"
+#include "LogVCamEditor.h"
 #include "Modifier/VCamModifier.h"
 
 #include "Interfaces/IPluginManager.h"
 #include "Kismet2/KismetEditorUtilities.h"
+#include "Modules/ModuleInterface.h"
+#include "Modules/ModuleManager.h"
 #include "PropertyEditorModule.h"
+#include "Customizations/OutputProvider/ConnectionRemapCustomization_VCamWidget.h"
+#include "Customizations/OutputProvider/OutputProviderTypeCustomization.h"
+#include "Output/VCamOutputViewport.h"
+#include "UI/VCamWidget.h"
+#include "UI/Switcher/VCamStateSwitcherWidget.h"
 
 namespace UE::VCamCoreEditor::Private
 {
-	class FVCamCoreEditorModule : public IModuleInterface
+	TSharedPtr<IConnectionRemapCustomization> FVCamCoreEditorModule::CreateConnectionRemapCustomization(TSubclassOf<UVCamWidget> Class) const
 	{
-	public:
+		for (UClass* const CurrentClass = Class; CurrentClass != UUserWidget::StaticClass(); Class = Class->GetSuperClass())
+		{
+			const FGetConnectionRemappingCustomization* FactoryDelegate = ConnectionRemappingCustomizationFactories.Find(Class);
+			if (FactoryDelegate)
+			{
+				return FactoryDelegate->Execute();
+			}
+		}
+		return nullptr;
+	}
 
-		// IModuleInterface interface
+	void FVCamCoreEditorModule::StartupModule()
+	{
+		RegisterAutoGeneratedDefaultEvents();
+		RegisterCustomizations();
+		RegisterDefaultConnectionRemappingCustomizations();
+	}
+	
+	void FVCamCoreEditorModule::ShutdownModule()
+	{
+		FKismetEditorUtilities::UnregisterAutoBlueprintNodeCreation(this);
+		UnregisterCustomizations();
+	}
 
-		virtual void StartupModule() override
+	void FVCamCoreEditorModule::RegisterConnectionRemapCustomization(TSubclassOf<UVCamWidget> Class, FGetConnectionRemappingCustomization GetterDelegate)
+	{
+		if (!Class)
 		{
-			RegisterAutoGeneratedDefaultEvents();
-		
-			FPropertyEditorModule& PropertyModule = FModuleManager::GetModuleChecked<FPropertyEditorModule>("PropertyEditor");
+			return;
+		}
 
-			PropertyModule.RegisterCustomPropertyTypeLayout(
-				FVCamInputProfile::StaticStruct()->GetFName(),
-				FOnGetPropertyTypeCustomizationInstance::CreateStatic(&FVCamInputProfileCustomization::MakeInstance)
-				);
-			PropertyModule.RegisterCustomPropertyTypeLayout(
-				FVCamConnectionTargetSettings::StaticStruct()->GetFName(),
-				FOnGetPropertyTypeCustomizationInstance::CreateStatic(&FConnectionTargetSettingsTypeCustomization::MakeInstance)
-				);
-			PropertyModule.RegisterCustomPropertyTypeLayout(
-				FWidgetConnectionConfig::StaticStruct()->GetFName(), 
-				FOnGetPropertyTypeCustomizationInstance::CreateStatic(&FWidgetConnectionConfigTypeCustomization::MakeInstance)
-				);
-			
-			PropertyModule.RegisterCustomClassLayout(
-				AVCamBaseActor::StaticClass()->GetFName(),
-				FOnGetDetailCustomizationInstance::CreateStatic(&FVCamBaseActorCustomization::MakeInstance)
+		if (Class->GetClass()->IsChildOf(UBlueprintGeneratedClass::StaticClass()))
+		{
+			UE_LOG(LogVCamEditor,
+				Error,
+				TEXT("RegisterConnectionTargetRemappingCustomizer is designed for native C++ base classes."
+					"Blueprints are untested and will likely break due to recompilation and reinstancing."
+					"'%s' is will not be registered!"
+					),
+				*Class->GetName()
 				);
-		};
-	
-		virtual void ShutdownModule() override
+			return;
+		}
+		
+		if (ConnectionRemappingCustomizationFactories.Contains(Class))
 		{
-			FPropertyEditorModule& PropertyModule = FModuleManager::GetModuleChecked<FPropertyEditorModule>("PropertyEditor");
-			PropertyModule.UnregisterCustomPropertyTypeLayout("VCamInputProfile");
-		};
+			UE_LOG(LogVCamEditor, Error, TEXT("Class '%s' is already registered!"), *Class->GetName());
+			return;
+		}
 
-		// Set up default nodes for VCam Core blueprint classes
-		void RegisterAutoGeneratedDefaultEvents()
+		ConnectionRemappingCustomizationFactories.Add(Class, MoveTemp(GetterDelegate));
+	}
+	
+	void FVCamCoreEditorModule::UnregisterConnectionRemapCustomization(TSubclassOf<UVCamWidget> Class)
+	{
+		const bool bIsInvalidClass = Class == UVCamWidget::StaticClass()
+			|| Class == UVCamStateSwitcherWidget::StaticClass();
+		if (bIsInvalidClass)
 		{
-			FKismetEditorUtilities::RegisterAutoGeneratedDefaultEvent(this, UVCamBlueprintModifier::StaticClass(), "OnInitialize");
-			FKismetEditorUtilities::RegisterAutoGeneratedDefaultEvent(this, UVCamBlueprintModifier::StaticClass(), "OnDeinitialize");
-			FKismetEditorUtilities::RegisterAutoGeneratedDefaultEvent(this, UVCamBlueprintModifier::StaticClass(), "OnApply");
-		}	
-	};
+			UE_LOG(LogVCamEditor, Error, TEXT("Class '%s' cannot be unregistered!"), Class ? *Class->GetName() : TEXT("null"));
+			return;
+		}
+
+		// Any alive instances will be cleared when the details panel redraws
+		ConnectionRemappingCustomizationFactories.Remove(Class);
+	}
+	
+	void FVCamCoreEditorModule::RegisterAutoGeneratedDefaultEvents()
+	{
+		FKismetEditorUtilities::RegisterAutoGeneratedDefaultEvent(this, UVCamBlueprintModifier::StaticClass(), "OnInitialize");
+		FKismetEditorUtilities::RegisterAutoGeneratedDefaultEvent(this, UVCamBlueprintModifier::StaticClass(), "OnDeinitialize");
+		FKismetEditorUtilities::RegisterAutoGeneratedDefaultEvent(this, UVCamBlueprintModifier::StaticClass(), "OnApply");
+	}
+	
+	void FVCamCoreEditorModule::RegisterCustomizations()
+	{
+		FPropertyEditorModule& PropertyModule = FModuleManager::GetModuleChecked<FPropertyEditorModule>("PropertyEditor");
+
+		PropertyModule.RegisterCustomPropertyTypeLayout(
+			FVCamInputProfile::StaticStruct()->GetFName(),
+			FOnGetPropertyTypeCustomizationInstance::CreateStatic(&FVCamInputProfileCustomization::MakeInstance)
+		);
+		PropertyModule.RegisterCustomPropertyTypeLayout(
+			FVCamConnectionTargetSettings::StaticStruct()->GetFName(),
+			FOnGetPropertyTypeCustomizationInstance::CreateStatic(&FConnectionTargetSettingsTypeCustomization::MakeInstance)
+		);
+		PropertyModule.RegisterCustomPropertyTypeLayout(
+			FWidgetConnectionConfig::StaticStruct()->GetFName(), 
+			FOnGetPropertyTypeCustomizationInstance::CreateStatic(&FWidgetConnectionConfigTypeCustomization::MakeInstance)
+		);
+		// Must be set up as IPropertyTypeCustomization as well to work on Instanced properties
+		PropertyModule.RegisterCustomPropertyTypeLayout(
+			UVCamOutputProviderBase::StaticClass()->GetFName(), 
+			FOnGetPropertyTypeCustomizationInstance::CreateStatic(&FOutputProviderTypeCustomization::MakeInstance)
+		);
+		
+		PropertyModule.RegisterCustomClassLayout(
+			UVCamOutputProviderBase::StaticClass()->GetFName(),
+			FOnGetDetailCustomizationInstance::CreateStatic(&FOutputProviderLayoutCustomization::MakeInstance)
+		);
+		PropertyModule.RegisterCustomClassLayout(
+			AVCamBaseActor::StaticClass()->GetFName(),
+			FOnGetDetailCustomizationInstance::CreateStatic(&FVCamBaseActorCustomization::MakeInstance)
+		);
+	}
+
+	void FVCamCoreEditorModule::RegisterDefaultConnectionRemappingCustomizations()
+	{
+		RegisterConnectionRemapCustomization(
+			UVCamWidget::StaticClass(),
+			FGetConnectionRemappingCustomization::CreateStatic(&FConnectionRemapCustomization_VCamWidget::Make)
+			);
+	}
+
+	void FVCamCoreEditorModule::UnregisterCustomizations()
+	{
+		FPropertyEditorModule& PropertyModule = FModuleManager::GetModuleChecked<FPropertyEditorModule>("PropertyEditor");
+		
+		PropertyModule.UnregisterCustomPropertyTypeLayout(FVCamInputProfile::StaticStruct()->GetFName());
+		PropertyModule.UnregisterCustomPropertyTypeLayout(FVCamConnectionTargetSettings::StaticStruct()->GetFName());
+		PropertyModule.UnregisterCustomPropertyTypeLayout(FWidgetConnectionConfig::StaticStruct()->GetFName());
+		PropertyModule.UnregisterCustomPropertyTypeLayout(UVCamOutputProviderBase::StaticClass()->GetFName());
+		
+		PropertyModule.UnregisterCustomClassLayout(UVCamOutputProviderBase::StaticClass()->GetFName());
+		PropertyModule.UnregisterCustomClassLayout(AVCamBaseActor::StaticClass()->GetFName());
+	}
 }
 
 IMPLEMENT_MODULE(UE::VCamCoreEditor::Private::FVCamCoreEditorModule, VCamCoreEditor);
\ No newline at end of file
diff --git a/Engine/Plugins/Experimental/VirtualCamera/Source/VCamCoreEditor/Private/VCamCoreEditorModule.h b/Engine/Plugins/Experimental/VirtualCamera/Source/VCamCoreEditor/Private/VCamCoreEditorModule.h
new file mode 100644
--- /dev/null
+++ b/Engine/Plugins/Experimental/VirtualCamera/Source/VCamCoreEditor/Private/VCamCoreEditorModule.h
@@ -0,0 +1,49 @@
+// Copyright Epic Games, Inc. All Rights Reserved.
+
+#pragma once
+
+#include "CoreMinimal.h"
+
+#include "IVCamCoreEditorModule.h"
+
+#include "UObject/WeakObjectPtr.h"
+#include "UObject/WeakObjectPtrTemplates.h"
+
+class UClass;
+class UVCamWidget;
+
+namespace UE::VCamCoreEditor::Private
+{
+	class FVCamCoreEditorModule : public IVCamCoreEditorModule
+	{
+	public:
+
+		static FVCamCoreEditorModule& Get()
+		{
+			return FModuleManager::Get().GetModuleChecked<FVCamCoreEditorModule>("VCamCoreEditor");
+		}
+
+		TSharedPtr<IConnectionRemapCustomization> CreateConnectionRemapCustomization(TSubclassOf<UVCamWidget> Class) const;
+
+		//~ Begin IModuleInterface Interface
+		virtual void StartupModule() override;
+		virtual void ShutdownModule() override;
+		//~ End IModuleInterface Interface
+		
+		//~ Begin IVCamCoreEditorModule Interface
+		virtual void RegisterConnectionRemapCustomization(TSubclassOf<UVCamWidget> Class, FGetConnectionRemappingCustomization GetterDelegate) override;
+		virtual void UnregisterConnectionRemapCustomization(TSubclassOf<UVCamWidget> Class) override;
+		//~ Begin IVCamCoreEditorModule Interface
+
+	private:
+
+		/** Maps pointer to UVCamWidget classes */
+		TMap<TWeakObjectPtr<UClass>, FGetConnectionRemappingCustomization> ConnectionRemappingCustomizationFactories;
+
+		void RegisterAutoGeneratedDefaultEvents();
+		void RegisterCustomizations();
+		void RegisterDefaultConnectionRemappingCustomizations();
+		
+		void UnregisterCustomizations();
+	};
+}
\ No newline at end of file
diff --git a/Engine/Plugins/Experimental/VirtualCamera/Source/VCamCoreEditor/Public/Customization/IConnectionRemapCustomization.h b/Engine/Plugins/Experimental/VirtualCamera/Source/VCamCoreEditor/Public/Customization/IConnectionRemapCustomization.h
new file mode 100644
--- /dev/null
+++ b/Engine/Plugins/Experimental/VirtualCamera/Source/VCamCoreEditor/Public/Customization/IConnectionRemapCustomization.h
@@ -0,0 +1,55 @@
+// Copyright Epic Games, Inc. All Rights Reserved.
+
+#pragma once
+
+#include "CoreMinimal.h"
+#include "UObject/WeakObjectPtr.h"
+#include "UObject/WeakObjectPtrTemplates.h"
+
+class IDetailGroup;
+class UVCamWidget;
+
+namespace UE::VCamCoreEditor
+{
+	class IConnectionRemapUtils;
+
+	struct FTargetConnectionDisplaySettings
+	{
+		bool bOnlyShowManuallyConfiguredConnections = true;
+	};
+	
+	struct FShouldGenerateArgs
+	{
+		TWeakObjectPtr<UVCamWidget> CustomizedWidget;
+		FTargetConnectionDisplaySettings DisplaySettings;
+	};
+
+	struct FConnectionRemapCustomizationArgs
+	{
+		/** Builder for the group under which to add the widget's data. */
+		IDetailGroup& WidgetGroup;
+		TSharedRef<IConnectionRemapUtils> Utils;
+		TWeakObjectPtr<UVCamWidget> CustomizedWidget;
+		
+		FTargetConnectionDisplaySettings DisplaySettings;
+	};
+
+	/**
+	 * Implements support for customizing the connection target points of specific widget types.
+	 *
+	 * Instances are created when a widget needs customization and kept alive until 1. the widget becomes invalid or 2. the details panel displays a new object;
+	 * stays alive across calls to IDetailBuilder::ForceRefreshDetails.
+	 */
+	class VCAMCOREEDITOR_API IConnectionRemapCustomization : public TSharedFromThis<IConnectionRemapCustomization>
+	{
+	public:
+
+		/** Whether a group should be generated for this widget */
+		virtual bool CanGenerateGroup(const FShouldGenerateArgs& Args) const = 0;
+		
+		/** Adds rows for CustomizedWidget */
+		virtual void Customize(const FConnectionRemapCustomizationArgs& Args) = 0;
+		
+		virtual ~IConnectionRemapCustomization() = default;
+	};
+}
\ No newline at end of file
diff --git a/Engine/Plugins/Experimental/VirtualCamera/Source/VCamCoreEditor/Public/Customization/IConnectionRemapUtils.h b/Engine/Plugins/Experimental/VirtualCamera/Source/VCamCoreEditor/Public/Customization/IConnectionRemapUtils.h
new file mode 100644
--- /dev/null
+++ b/Engine/Plugins/Experimental/VirtualCamera/Source/VCamCoreEditor/Public/Customization/IConnectionRemapUtils.h
@@ -0,0 +1,61 @@
+// Copyright Epic Games, Inc. All Rights Reserved.
+
+#pragma once
+
+#include "CoreMinimal.h"
+#include "UI/VCamConnectionStructs.h"
+
+class IDetailGroup;
+class UVCamWidget;
+struct FVCamConnectionTargetSettings;
+
+namespace UE::VCamCoreEditor
+{
+	DECLARE_DELEGATE_OneParam(FOnTargetSettingsChanged, const FVCamConnectionTargetSettings& NewSettings);
+	struct FAddConnectionArgs
+	{
+		IDetailGroup& DetailGroup;
+
+		FName ConnectionName;
+
+		/** Name of the property.  */
+		FName PropertyName;
+		/** A containing a FVCamConnectionTargetSettings property with name PropertyName */
+		FStructOnScope StructData;
+
+		/** Called when the settings are changed. Copy the passed settings to your UPROPERTY(). */
+		FOnTargetSettingsChanged OnTargetSettingsChangedDelegate;
+
+		/** The font to use for displaying property texts. */
+		FSlateFontInfo Font;
+
+		FAddConnectionArgs(IDetailGroup& DetailGroup, FName ConnectionName, FName PropertyName, FStructOnScope StructData, FOnTargetSettingsChanged OnTargetSettingsChangedDelegate, FSlateFontInfo Font)
+			: DetailGroup(DetailGroup)
+			, ConnectionName(ConnectionName)
+			, PropertyName(PropertyName)
+			, StructData(MoveTemp(StructData))
+			, OnTargetSettingsChangedDelegate(MoveTemp(OnTargetSettingsChangedDelegate))
+			, Font(MoveTemp(Font))
+		{}
+	};
+	
+	/** Passed to IConnectionTargetRemappingCustomizers to re-use functionality. */
+	class VCAMCOREEDITOR_API IConnectionRemapUtils : public TSharedFromThis<IConnectionRemapUtils>
+	{
+	public:
+
+		/** Adds a property row representing TargetSettings to DetailGroup. */
+		virtual void AddConnection(FAddConnectionArgs Params) = 0;
+
+		/** @return the font used for properties and details */
+		virtual FSlateFontInfo GetRegularFont() const = 0;
+
+		/**
+		 * Refreshes the details view and regenerates all the customized layouts
+		 * Use only when you need to remove or add complicated dynamic items
+		 */
+		virtual void ForceRefreshProperties() const = 0;
+		
+		virtual ~IConnectionRemapUtils() = default;
+	};
+}
diff --git a/Engine/Plugins/Experimental/VirtualCamera/Source/VCamCoreEditor/Public/IVCamCoreEditorModule.h b/Engine/Plugins/Experimental/VirtualCamera/Source/VCamCoreEditor/Public/IVCamCoreEditorModule.h
new file mode 100644
--- /dev/null
+++ b/Engine/Plugins/Experimental/VirtualCamera/Source/VCamCoreEditor/Public/IVCamCoreEditorModule.h
@@ -0,0 +1,41 @@
+// Copyright Epic Games, Inc. All Rights Reserved.
+
+#pragma once
+
+#include "CoreMinimal.h"
+#include "Modules/ModuleInterface.h"
+#include "Modules/ModuleManager.h"
+#include "Templates/SubclassOf.h"
+
+class UVCamWidget;
+
+namespace UE::VCamCoreEditor
+{
+	class IConnectionRemapCustomization;
+	
+	DECLARE_DELEGATE_RetVal(TSharedRef<IConnectionRemapCustomization>, FGetConnectionRemappingCustomization)
+	
+	/** Passed to IConnectionTargetRemappingCustomizers to re-use functionality. */
+	class VCAMCOREEDITOR_API IVCamCoreEditorModule : public IModuleInterface
+	{
+	public:
+
+		static IVCamCoreEditorModule& Get()
+		{
+			return FModuleManager::Get().GetModuleChecked<IVCamCoreEditorModule>("VCamCoreEditor");
+		}
+
+		/**
+		 * Registers a specific remapping customization for a class. At most one customization is used per VCamWidget class.
+		 * To determine a widget's customization, the class's hierarchy is navigated recursively until a customization is found that is registered with a parent class.
+		 * Cannot override already registered customizations. You'd have to call UnregisterConnectionTargetRemappingCustomizer first.
+		 * 
+		 * @param Class The subclass to register for - cannot be UVCamWidget and
+		 * @param GetterDelegate The factory to invoke whenever a widget is customized. 
+		 */
+		virtual void RegisterConnectionRemapCustomization(TSubclassOf<UVCamWidget> Class, FGetConnectionRemappingCustomization GetterDelegate) = 0;
+		virtual void UnregisterConnectionRemapCustomization(TSubclassOf<UVCamWidget> Class) = 0;
+		
+		virtual ~IVCamCoreEditorModule() override = default;
+	};
+}
diff --git a/Engine/Plugins/Experimental/VirtualCamera/Source/VCamCoreEditor/VCamCoreEditor.build.cs b/Engine/Plugins/Experimental/VirtualCamera/Source/VCamCoreEditor/VCamCoreEditor.build.cs
--- a/Engine/Plugins/Experimental/VirtualCamera/Source/VCamCoreEditor/VCamCoreEditor.build.cs
+++ b/Engine/Plugins/Experimental/VirtualCamera/Source/VCamCoreEditor/VCamCoreEditor.build.cs
@@ -34,7 +34,8 @@ public class VCamCoreEditor : ModuleRules
 				"UMGEditor",
 				"UnrealEd",
 				"VCamCore",
-				"VCamInput"
+				"VCamInput",
+				"VPUtilities"
 			}
 		);
 	}
