Cannot add new output providers

diff --git a/Engine/Plugins/Experimental/VirtualCamera/Source/VCamCoreEditor/Private/Customizations/OutputProvider/OutputProviderTypeCustomization.cpp b/Engine/Plugins/Experimental/VirtualCamera/Source/VCamCoreEditor/Private/Customizations/OutputProvider/OutputProviderTypeCustomization.cpp
--- a/Engine/Plugins/Experimental/VirtualCamera/Source/VCamCoreEditor/Private/Customizations/OutputProvider/OutputProviderTypeCustomization.cpp
+++ b/Engine/Plugins/Experimental/VirtualCamera/Source/VCamCoreEditor/Private/Customizations/OutputProvider/OutputProviderTypeCustomization.cpp
@@ -31,6 +31,15 @@ namespace UE::VCamCoreEditor::Private
 			: nullptr;
 		if (!CustomizedOutputProvider)
 		{
+			ChildBuilder.AddCustomRow(FText::GetEmpty())
+				.NameContent()
+				[
+					PropertyHandle->CreatePropertyNameWidget()
+				]
+				.ValueContent()
+				[
+					PropertyHandle->CreatePropertyValueWidget()
+				];
 			return;
 		}
 		
