Add new Blueprint functions:
- Get Connection By Name
- Is Connected By Name
- Get Connection Point Name By Name
- Get Connected Modifier By Name
- Get Connected Input Action By Name

diff --git a/Engine/Plugins/Experimental/VirtualCamera/Source/VCamCore/Private/UI/VCamUIFunctionLibrary.cpp b/Engine/Plugins/Experimental/VirtualCamera/Source/VCamCore/Private/UI/VCamUIFunctionLibrary.cpp
--- a/Engine/Plugins/Experimental/VirtualCamera/Source/VCamCore/Private/UI/VCamUIFunctionLibrary.cpp
+++ b/Engine/Plugins/Experimental/VirtualCamera/Source/VCamCore/Private/UI/VCamUIFunctionLibrary.cpp
@@ -2,6 +2,8 @@
 
 #include "UI/VCamUIFunctionLibrary.h"
 
+#include "UI/VCamWidget.h"
+
 bool UVCamUIFunctionLibrary::IsConnected_VCamConnection(const FVCamConnection& Connection)
 {
 	return Connection.IsConnected();
@@ -22,4 +24,63 @@ UInputAction* UVCamUIFunctionLibrary::GetConnectedInputAction_VCamConnection(con
 	return Connection.ConnectedAction;
 }
 
+bool UVCamUIFunctionLibrary::GetConnectionByName_VCamWidget(UVCamWidget* Widget, FName ConnectionId, FVCamConnection& OutConnection)
+{
+	if (!Widget)
+	{
+		return false;
+	}
+
+	if (const FVCamConnection* Connection = Widget->Connections.Find(ConnectionId))
+	{
+		OutConnection = *Connection;
+		return true;
+	}
+	return false;
+}
+
+bool UVCamUIFunctionLibrary::IsConnected_VCamWidget(UVCamWidget* Widget, FName ConnectionId, bool& bOutIsConnected)
+{
+	FVCamConnection Connection;
+	if (GetConnectionByName_VCamWidget(Widget, ConnectionId, Connection))
+	{
+		bOutIsConnected = IsConnected_VCamConnection(Connection);
+		return true;
+	}
+	return false;
+}
+
+bool UVCamUIFunctionLibrary::GetConnectionPointName_VCamWidget(UVCamWidget* Widget, FName ConnectionId, FName& OutConnectionPointName)
+{
+	FVCamConnection Connection;
+	if (GetConnectionByName_VCamWidget(Widget, ConnectionId, Connection))
+	{
+		OutConnectionPointName = GetConnectionPointName_VCamConnection(Connection);
+		return true;
+	}
+	return false;
+}
+
+bool UVCamUIFunctionLibrary::GetConnectedModifier_VCamWidget(UVCamWidget* Widget, FName ConnectionId, UVCamModifier*& OutModifier)
+{
+	FVCamConnection Connection;
+	if (GetConnectionByName_VCamWidget(Widget, ConnectionId, Connection))
+	{
+		OutModifier = GetConnectedModifier_VCamConnection(Connection);
+		return true;
+	}
+	return false;
+}
+
+bool UVCamUIFunctionLibrary::GetConnectedInputAction_VCamWidget(UVCamWidget* Widget, FName ConnectionId, UInputAction*& OutInputAction)
+{
+	FVCamConnection Connection;
+	if (GetConnectionByName_VCamWidget(Widget, ConnectionId, Connection))
+	{
+		OutInputAction = GetConnectedInputAction_VCamConnection(Connection);
+		return true;
+	}
+	return false;
+}
+
 
diff --git a/Engine/Plugins/Experimental/VirtualCamera/Source/VCamCore/Public/UI/VCamUIFunctionLibrary.h b/Engine/Plugins/Experimental/VirtualCamera/Source/VCamCore/Public/UI/VCamUIFunctionLibrary.h
--- a/Engine/Plugins/Experimental/VirtualCamera/Source/VCamCore/Public/UI/VCamUIFunctionLibrary.h
+++ b/Engine/Plugins/Experimental/VirtualCamera/Source/VCamCore/Public/UI/VCamUIFunctionLibrary.h
@@ -10,6 +10,7 @@
 
 class UInputAction;
 class UVCamModifier;
+class UVCamWidget;
 
 /*
  * Function Library primarily used for exposing functionality of UI related structs to Blueprints 
@@ -49,4 +50,26 @@ public:
 	 */
 	UFUNCTION(BlueprintPure, Category = "VCam Connections", meta = (DisplayName = "Get Connected Input Action"))
 	static UInputAction* GetConnectedInputAction_VCamConnection(const FVCamConnection& Connection);
+
+
+	
+	/** Gets a connection by its ID.  */
+	UFUNCTION(BlueprintPure, Category = "VCam Connections", meta = (DisplayName = "Get Connection By Name", DefaultToSelf = "Widget"))
+	static bool GetConnectionByName_VCamWidget(UVCamWidget* Widget, FName ConnectionId, FVCamConnection& OutConnection);
+
+	/** Gets whether the ConnectionId is connected on the widget. */
+	UFUNCTION(BlueprintPure, Category = "VCam Connections", meta = (DisplayName = "Is Connected By Name", DefaultToSelf = "Widget"))
+	static bool IsConnected_VCamWidget(UVCamWidget* Widget, FName ConnectionId, bool& bOutIsConnected);
+	
+	/** Gets the connection point the connection ConnectionId attempts to connect to on the widget. */
+	UFUNCTION(BlueprintPure, Category = "VCam Connections", meta = (DisplayName = "Get Connection Point Name By Name", DefaultToSelf = "Widget"))
+	static bool GetConnectionPointName_VCamWidget(UVCamWidget* Widget, FName ConnectionId, FName& OutConnectionPointName);
+	
+	/** Gets the modifier connected to ConnectionId on the widget. */
+	UFUNCTION(BlueprintPure, Category = "VCam Connections", meta = (DisplayName = "Get Connected Modifier By Name", DefaultToSelf = "Widget"))
+	static bool GetConnectedModifier_VCamWidget(UVCamWidget* Widget, FName ConnectionId, UVCamModifier*& OutModifier);
+
+	/** Gets the input action mapped to ConnectionId on the widget. */
+	UFUNCTION(BlueprintPure, Category = "VCam Connections", meta = (DisplayName = "Get Connected Input Action By Name", DefaultToSelf = "Widget"))
+	static bool GetConnectedInputAction_VCamWidget(UVCamWidget* Widget, FName ConnectionId, UInputAction*& OutInputAction);
 };
\ No newline at end of file
