Customize viewport locker

diff --git a/Engine/Plugins/Experimental/VirtualCamera/Source/VCamCore/Private/Output/VCamOutputProviderBase.cpp b/Engine/Plugins/Experimental/VirtualCamera/Source/VCamCore/Private/Output/VCamOutputProviderBase.cpp
--- a/Engine/Plugins/Experimental/VirtualCamera/Source/VCamCore/Private/Output/VCamOutputProviderBase.cpp
+++ b/Engine/Plugins/Experimental/VirtualCamera/Source/VCamCore/Private/Output/VCamOutputProviderBase.cpp
@@ -268,6 +268,30 @@ void UVCamOutputProviderBase::DestroyUMG()
 	}
 }
 
+void UVCamOutputProviderBase::SuspendOutput()
+{
+	if (IsActive())
+	{
+		bWasActive = true;
+		SetActive(false);
+	}
+}
+
+void UVCamOutputProviderBase::RestoreOutput()
+{
+	if (bWasActive && !IsActive())
+	{
+		SetActive(true);
+	}
+	bWasActive = false;
+}
+
+bool UVCamOutputProviderBase::NeedsForceLockToViewport() const
+{
+	// The widget is displayed via a post process material, which is applied to the camera's post process settings, hence anything will only be visible when locked.
+	return DisplayType == EVPWidgetDisplayType::PostProcess;
+}
+
 void UVCamOutputProviderBase::NotifyWidgetOfComponentChange() const
 {
 	if (UMGWidget && UMGWidget->IsDisplayed())
diff --git a/Engine/Plugins/Experimental/VirtualCamera/Source/VCamCore/Private/Util/LevelViewportUtils.cpp b/Engine/Plugins/Experimental/VirtualCamera/Source/VCamCore/Private/Util/LevelViewportUtils.cpp
--- a/Engine/Plugins/Experimental/VirtualCamera/Source/VCamCore/Private/Util/LevelViewportUtils.cpp
+++ b/Engine/Plugins/Experimental/VirtualCamera/Source/VCamCore/Private/Util/LevelViewportUtils.cpp
@@ -20,6 +20,14 @@ namespace UE::VCamCore::LevelViewportUtils::Private
 {
 	namespace Locking
 	{
+		enum class EViewportFlags
+		{
+			None,
+			IsNotUsed   = 1 << 0,
+			ForceUse	= 1 << 1,
+		};
+		ENUM_CLASS_FLAGS(EViewportFlags)
+		
 #if WITH_EDITOR
 		static void UpdateLockStateForEditor(FVCamViewportLockState& ViewportLockState, EVCamTargetViewportID ViewportID, bool bNewLockState, AActor* ActorToLockWith = nullptr)
 		{
@@ -87,9 +95,15 @@ namespace UE::VCamCore::LevelViewportUtils::Private
 			}
 		}
 		
-		static void UpdateLockState(FVCamViewportLockState& ViewportLockState, EVCamTargetViewportID ViewportID, bool bForcfullySetUnlocked, AActor* ActorToLockWith = nullptr)
+		static void UpdateLockState(FVCamViewportLockState& ViewportLockState, EVCamTargetViewportID ViewportID, EViewportFlags ViewportFlags, AActor* ActorToLockWith = nullptr)
 		{
-			const bool bNewLockState = ViewportLockState.bLockViewportToCamera & !bForcfullySetUnlocked;
+			const bool bForceUse = (ViewportFlags & EViewportFlags::ForceUse) != EViewportFlags::None;
+			const bool bIsUsed = (ViewportFlags & EViewportFlags::IsNotUsed) == EViewportFlags::None;
+			const bool bNewLockState = bForceUse || (ViewportLockState.bLockViewportToCamera && bIsUsed);
+
+#if WITH_EDITOR
+			ViewportLockState.bIsForceLocked = bForceUse;
+#endif
 			for (const FWorldContext& Context : GEngine->GetWorldContexts())
 			{
 #if WITH_EDITOR
@@ -109,16 +123,28 @@ namespace UE::VCamCore::LevelViewportUtils::Private
 	void UpdateViewportLocksFromOutputs(TArray<TObjectPtr<UVCamOutputProviderBase>> OutputProviders, FVCamViewportLocker& LockData, AActor* ActorToLockWith)
 	{
 		TSet<EVCamTargetViewportID> Viewports;
+		TSet<EVCamTargetViewportID> ForcefullyLockedViewports;
+		
 		Algo::TransformIf(OutputProviders, Viewports,
 			[](TObjectPtr<UVCamOutputProviderBase> Output){ return Output && Output->IsActive(); },
 			[](TObjectPtr<UVCamOutputProviderBase> Output) { return Output->GetTargetViewport(); }
 			);
+		Algo::TransformIf(OutputProviders, ForcefullyLockedViewports,
+			[](TObjectPtr<UVCamOutputProviderBase> Output){ return Output && Output->IsActive() && Output->NeedsForceLockToViewport(); },
+			[](TObjectPtr<UVCamOutputProviderBase> Output) { return Output->GetTargetViewport(); }
+			);
 		
 		check(LockData.Locks.Num() == 4);
 		for (TPair<EVCamTargetViewportID, FVCamViewportLockState>& ViewportData : LockData.Locks)
 		{
-			const bool bShouldUnlock = !Viewports.Contains(ViewportData.Key);
-			Locking::UpdateLockState(ViewportData.Value, ViewportData.Key, bShouldUnlock, ActorToLockWith);
+			const bool bIsForceLocked = ForcefullyLockedViewports.Contains(ViewportData.Key);
+			const bool bIsNotUsed = !Viewports.Contains(ViewportData.Key);
+			
+			Locking::EViewportFlags Flags = Locking::EViewportFlags::None;
+			Flags |= bIsForceLocked ? Locking::EViewportFlags::ForceUse : Locking::EViewportFlags::None;
+			Flags |= bIsNotUsed ? Locking::EViewportFlags::IsNotUsed : Locking::EViewportFlags::None;
+			
+			Locking::UpdateLockState(ViewportData.Value, ViewportData.Key, Flags, ActorToLockWith);
 		}
 	}
 
@@ -126,7 +152,7 @@ namespace UE::VCamCore::LevelViewportUtils::Private
 	{
 		for (TPair<EVCamTargetViewportID, FVCamViewportLockState>& ViewportData : LockData.Locks)
 		{
-			Locking::UpdateLockState(ViewportData.Value, ViewportData.Key, false);
+			Locking::UpdateLockState(ViewportData.Value, ViewportData.Key, Locking::EViewportFlags::IsNotUsed);
 		}
 	}
 
diff --git a/Engine/Plugins/Experimental/VirtualCamera/Source/VCamCore/Public/Output/VCamOutputProviderBase.h b/Engine/Plugins/Experimental/VirtualCamera/Source/VCamCore/Public/Output/VCamOutputProviderBase.h
--- a/Engine/Plugins/Experimental/VirtualCamera/Source/VCamCore/Public/Output/VCamOutputProviderBase.h
+++ b/Engine/Plugins/Experimental/VirtualCamera/Source/VCamCore/Public/Output/VCamOutputProviderBase.h
@@ -90,24 +90,12 @@ public:
 	UVPFullScreenUserWidget* GetUMGWidget() { return UMGWidget; };
 
 	/** Temporarily disable the output.  Caller must eventually call RestoreOutput. */
-	void SuspendOutput()
-	{
-		if (IsActive())
-		{
-			bWasActive = true;
-			SetActive(false);
-		}
-	}
-
+	void SuspendOutput();
 	/** Restore the output state from previous call to disable output. */
-	void RestoreOutput()
-	{
-		if (bWasActive && !IsActive())
-		{
-			SetActive(true);
-		}
-		bWasActive = false;
-	}
+	void RestoreOutput();
+	
+	/** @return Whether this output provider should requires the viewport to be locked to the camera in order to function correctly. */
+	bool NeedsForceLockToViewport() const;
 
 	/** Calls the VCamModifierInterface on the widget if it exists and also requests any child VCam Widgets to reconnect */
 	void NotifyWidgetOfComponentChange() const;
@@ -148,7 +136,6 @@ protected:
 	void RestoreOverrideResolutionForViewport(EVCamTargetViewportID ViewportToRestore);
 	/** Applies OverrideResolution to the passed in viewport - bUseOverrideResolution was already checked. */
 	void ApplyOverrideResolutionForViewport(EVCamTargetViewportID Viewport);
-
 	void ReapplyOverrideResolution(EVCamTargetViewportID Viewport);
 
 	void DisplayUMG();
diff --git a/Engine/Plugins/Experimental/VirtualCamera/Source/VCamCore/Public/Util/VCamViewportLocker.h b/Engine/Plugins/Experimental/VirtualCamera/Source/VCamCore/Public/Util/VCamViewportLocker.h
--- a/Engine/Plugins/Experimental/VirtualCamera/Source/VCamCore/Public/Util/VCamViewportLocker.h
+++ b/Engine/Plugins/Experimental/VirtualCamera/Source/VCamCore/Public/Util/VCamViewportLocker.h
@@ -12,12 +12,18 @@ struct FVCamViewportLockState
 	GENERATED_BODY()
 	
 	/** Whether the user wants the viewport to be locked */
-	UPROPERTY(EditAnywhere, Category = "Viewport")
+	UPROPERTY(EditAnywhere, Category = "Viewport", meta = (EditCondition = "!bIsForceLocked"))
 	bool bLockViewportToCamera = false;
 	
 	/** Whether this viewport is currently locked */
 	UPROPERTY(Transient)
 	bool bIsLockedToViewport = false;
+
+#if WITH_EDITORONLY_DATA
+	// This property is editor-only because we use it for EditCondition only
+	UPROPERTY(Transient)
+	bool bIsForceLocked = false;
+#endif
 	
 	/** Used for editor */
 	UPROPERTY(Transient)
diff --git a/Engine/Plugins/Experimental/VirtualCamera/Source/VCamCoreEditor/Private/Customizations/VCamViewportLockerTypeCustomization.cpp b/Engine/Plugins/Experimental/VirtualCamera/Source/VCamCoreEditor/Private/Customizations/VCamViewportLockerTypeCustomization.cpp
new file mode 100644
--- /dev/null
+++ b/Engine/Plugins/Experimental/VirtualCamera/Source/VCamCoreEditor/Private/Customizations/VCamViewportLockerTypeCustomization.cpp
@@ -0,0 +1,78 @@
+// Copyright Epic Games, Inc. All Rights Reserved.
+
+#include "VCamViewportLockerTypeCustomization.h"
+
+#include "DetailWidgetRow.h"
+#include "IDetailChildrenBuilder.h"
+#include "Util/VCamViewportLocker.h"
+#include "Widgets/Text/STextBlock.h"
+
+#define LOCTEXT_NAMESPACE "FVCamViewportLockerTypeCustomization"
+
+namespace UE::VCamCoreEditor::Private
+{
+	static FText GetViewportDisplayNameByIndex(int32 ViewportIndex)
+	{
+		switch (ViewportIndex)
+		{
+		case 0: return LOCTEXT("Viewport1", "Viewport 1");
+		case 1: return LOCTEXT("Viewport1", "Viewport 2");
+		case 2: return LOCTEXT("Viewport1", "Viewport 3");
+		case 3: return LOCTEXT("Viewport1", "Viewport 4");
+
+		default:
+			checkNoEntry();
+			return FText::GetEmpty();
+		}
+	}
+	
+	TSharedRef<IPropertyTypeCustomization> FVCamViewportLockerTypeCustomization::MakeInstance()
+	{
+		return MakeShared<FVCamViewportLockerTypeCustomization>();
+	}
+
+	void FVCamViewportLockerTypeCustomization::CustomizeHeader(TSharedRef<IPropertyHandle> PropertyHandle, FDetailWidgetRow& HeaderRow, IPropertyTypeCustomizationUtils& CustomizationUtils)
+	{
+		HeaderRow.
+			NameContent()
+			[
+				PropertyHandle->CreatePropertyNameWidget()
+			]
+			.ValueContent()
+			[
+				PropertyHandle->CreatePropertyValueWidget()
+			];
+	}
+
+	void FVCamViewportLockerTypeCustomization::CustomizeChildren(TSharedRef<IPropertyHandle> PropertyHandle, IDetailChildrenBuilder& ChildBuilder, IPropertyTypeCustomizationUtils& CustomizationUtils)
+	{
+		const FName LocksProperty = GET_MEMBER_NAME_CHECKED(FVCamViewportLocker, Locks);
+		const TSharedPtr<IPropertyHandle> LocksMap = PropertyHandle->GetChildHandle(LocksProperty);
+
+		uint32 NumElements = 0;
+		if (LocksMap->AsMap()->GetNumElements(NumElements) != FPropertyAccess::Success)
+		{
+			return;
+		}
+		
+		for (uint32 ChildIndex = 0; ChildIndex < NumElements; ++ChildIndex)
+		{
+			const TSharedPtr<IPropertyHandle> ViewportLockStateProperty = LocksMap->GetChildHandle(ChildIndex);
+			const TSharedPtr<IPropertyHandle> LockViewportToCameraProperty = ViewportLockStateProperty->GetChildHandle(GET_MEMBER_NAME_CHECKED(FVCamViewportLockState, bLockViewportToCamera));
+			ChildBuilder.AddProperty(LockViewportToCameraProperty.ToSharedRef())
+				.CustomWidget()
+				.NameContent()
+				[
+					SNew(STextBlock)
+					.Text(GetViewportDisplayNameByIndex(ChildIndex))
+					.Font(CustomizationUtils.GetRegularFont())
+				]
+				.ValueContent()
+				[
+					LockViewportToCameraProperty->CreatePropertyValueWidget()
+				];
+		}
+	}
+}
+
+#undef LOCTEXT_NAMESPACE
\ No newline at end of file
diff --git a/Engine/Plugins/Experimental/VirtualCamera/Source/VCamCoreEditor/Private/Customizations/VCamViewportLockerTypeCustomization.h b/Engine/Plugins/Experimental/VirtualCamera/Source/VCamCoreEditor/Private/Customizations/VCamViewportLockerTypeCustomization.h
new file mode 100644
--- /dev/null
+++ b/Engine/Plugins/Experimental/VirtualCamera/Source/VCamCoreEditor/Private/Customizations/VCamViewportLockerTypeCustomization.h
@@ -0,0 +1,26 @@
+// Copyright Epic Games, Inc. All Rights Reserved.
+
+#pragma once
+
+#include "CoreMinimal.h"
+#include "IPropertyTypeCustomization.h"
+
+class UWidgetTree;
+class UUserWidget;
+class IDetailPropertyRow;
+
+namespace UE::VCamCoreEditor::Private
+{
+	/** Just inlines the TMap (since the keys cannot be changed). */
+	class FVCamViewportLockerTypeCustomization : public IPropertyTypeCustomization
+	{
+	public:
+
+		static TSharedRef<IPropertyTypeCustomization> MakeInstance();
+
+		//~ Begin IStructCustomization Interface
+		virtual void CustomizeHeader(TSharedRef<IPropertyHandle> PropertyHandle, FDetailWidgetRow& HeaderRow, IPropertyTypeCustomizationUtils& CustomizationUtils) override;
+		virtual void CustomizeChildren(TSharedRef<IPropertyHandle> PropertyHandle, IDetailChildrenBuilder& ChildBuilder, IPropertyTypeCustomizationUtils& CustomizationUtils) override;
+		//~ End IStructCustomization Interface
+	};
+}
diff --git a/Engine/Plugins/Experimental/VirtualCamera/Source/VCamCoreEditor/Private/VCamCoreEditorModule.cpp b/Engine/Plugins/Experimental/VirtualCamera/Source/VCamCoreEditor/Private/VCamCoreEditorModule.cpp
--- a/Engine/Plugins/Experimental/VirtualCamera/Source/VCamCoreEditor/Private/VCamCoreEditorModule.cpp
+++ b/Engine/Plugins/Experimental/VirtualCamera/Source/VCamCoreEditor/Private/VCamCoreEditorModule.cpp
@@ -19,6 +19,7 @@
 #include "Modules/ModuleInterface.h"
 #include "Modules/ModuleManager.h"
 #include "PropertyEditorModule.h"
+#include "Customizations/VCamViewportLockerTypeCustomization.h"
 #include "Customizations/OutputProvider/ConnectionRemapCustomization_VCamWidget.h"
 #include "Customizations/OutputProvider/OutputProviderTypeCustomization.h"
 #include "Output/VCamOutputViewport.h"
@@ -124,6 +125,10 @@ namespace UE::VCamCoreEditor::Private
 			UVCamOutputProviderBase::StaticClass()->GetFName(), 
 			FOnGetPropertyTypeCustomizationInstance::CreateStatic(&FOutputProviderTypeCustomization::MakeInstance)
 		);
+		PropertyModule.RegisterCustomPropertyTypeLayout(
+			FVCamViewportLocker::StaticStruct()->GetFName(), 
+			FOnGetPropertyTypeCustomizationInstance::CreateStatic(&FVCamViewportLockerTypeCustomization::MakeInstance)
+		);
 		
 		PropertyModule.RegisterCustomClassLayout(
 			UVCamOutputProviderBase::StaticClass()->GetFName(),
@@ -151,6 +156,7 @@ namespace UE::VCamCoreEditor::Private
 		PropertyModule.UnregisterCustomPropertyTypeLayout(FVCamConnectionTargetSettings::StaticStruct()->GetFName());
 		PropertyModule.UnregisterCustomPropertyTypeLayout(FWidgetConnectionConfig::StaticStruct()->GetFName());
 		PropertyModule.UnregisterCustomPropertyTypeLayout(UVCamOutputProviderBase::StaticClass()->GetFName());
+		PropertyModule.UnregisterCustomPropertyTypeLayout(FVCamViewportLocker::StaticStruct()->GetFName());
 		
 		PropertyModule.UnregisterCustomClassLayout(UVCamOutputProviderBase::StaticClass()->GetFName());
 		PropertyModule.UnregisterCustomClassLayout(AVCamBaseActor::StaticClass()->GetFName());
