Most drop-downs for connection target settings now respect RequiredInterfaces (except for details customization for output provider)

diff --git a/Engine/Plugins/Experimental/VirtualCamera/Source/VCamCore/Private/UI/VCamConnectionStructs.cpp b/Engine/Plugins/Experimental/VirtualCamera/Source/VCamCore/Private/UI/VCamConnectionStructs.cpp
--- a/Engine/Plugins/Experimental/VirtualCamera/Source/VCamCore/Private/UI/VCamConnectionStructs.cpp
+++ b/Engine/Plugins/Experimental/VirtualCamera/Source/VCamCore/Private/UI/VCamConnectionStructs.cpp
@@ -15,56 +15,55 @@ bool FVCamConnection::IsConnected() const
 
 bool FVCamConnection::AttemptConnection(UVCamComponent* VCamComponent)
 {
-	bool bIsValidModifierConnection = false;
-		UVCamModifier* Modifier = VCamComponent->GetModifierByName(ConnectionTargetSettings.TargetModifierName);
-		FVCamModifierConnectionPoint* ConnectionPoint = nullptr;
-		if (Modifier)
-		{			
-			const bool bImplementsAllRequiredInterfaces = Algo::AllOf(RequiredInterfaces, [Modifier](const TSubclassOf<UInterface>& Interface)
-			{
-				return Modifier->GetClass()->ImplementsInterface(Interface);
-			});
+	UVCamModifier* Modifier = VCamComponent->GetModifierByName(ConnectionTargetSettings.TargetModifierName);
+	
+	UE_CLOG(!Modifier, LogVCamConnection, Warning, TEXT("Connection Failed: Failed to find Modifier with name %s in VCam Component"), *ConnectionTargetSettings.TargetModifierName.ToString());
+	constexpr bool bLogWarnings = true;
+	if (Modifier && IsConnectionValid(*Modifier, ConnectionTargetSettings.TargetConnectionPoint, bLogWarnings))
+	{
+		ConnectedModifier = Modifier;
+		ConnectedAction = Modifier->ConnectionPoints[ConnectionTargetSettings.TargetConnectionPoint].AssociatedAction;
+		return true;
+	}
 
-			if (bImplementsAllRequiredInterfaces)
-			{
-				ConnectionPoint = Modifier->ConnectionPoints.Find(ConnectionTargetSettings.TargetConnectionPoint);
-				if (ConnectionPoint)
-				{
-					// The Connection Point is considered valid if either we don't require an action or the required ActionType matches the ActionType inside the Connection Point
-					const bool bConnectionPointHasRequiredActionType = !bRequiresInputAction || (ConnectionPoint->AssociatedAction && ConnectionPoint->AssociatedAction->ValueType == ActionType);
-					if (bConnectionPointHasRequiredActionType)
-					{
-						bIsValidModifierConnection = true;
-					}
-					else
-					{
-						UE_LOG(LogVCamConnection, Warning, TEXT("Connection Failed: Connection Point with name %s does not provide an action of the required type. Expected %s but found %s"),
-							*ConnectionTargetSettings.TargetConnectionPoint.ToString(),
-							*UEnum::GetValueAsString(ActionType),
-							ConnectionPoint->AssociatedAction ? *UEnum::GetValueAsString(ConnectionPoint->AssociatedAction->ValueType) : TEXT("None"));
-					}
-				}
-				else
-				{
-					UE_LOG(LogVCamConnection, Warning, TEXT("Connection Failed: Failed to find Connection Point with name %s in Modifier %s"),
-						*ConnectionTargetSettings.TargetConnectionPoint.ToString(),
-						*ConnectionTargetSettings.TargetModifierName.ToString());
-				}
-			}
-			else
-			{
-				UE_LOG(LogVCamConnection, Warning, TEXT("Connection Failed: Modifier with name %s does not implement all required interfaces for this connection."), *ConnectionTargetSettings.TargetModifierName.ToString());
-			}
-		}
-		else
-		{
-			UE_LOG(LogVCamConnection, Warning, TEXT("Connection Failed: Failed to find Modifier with name %s in VCam Component"), *ConnectionTargetSettings.TargetModifierName.ToString());
-		}
-
-		ConnectedModifier = bIsValidModifierConnection ? Modifier : nullptr;
-		ConnectedAction = bIsValidModifierConnection ? ConnectionPoint->AssociatedAction : nullptr;
+	return false;
+}
 
-	return bIsValidModifierConnection;
+bool FVCamConnection::IsConnectionValid(UVCamModifier& Modifier, FName ConnectionPointName, bool bLogWarnings) const
+{
+	const bool bImplementsAllRequiredInterfaces = Algo::AllOf(RequiredInterfaces, [&Modifier](const TSubclassOf<UInterface>& Interface)
+	{
+		return Modifier.GetClass()->ImplementsInterface(Interface);
+	});
+	if (!bImplementsAllRequiredInterfaces)
+	{
+		UE_CLOG(bLogWarnings, LogVCamConnection, Warning, TEXT("Connection Failed: Modifier with name %s does not implement all required interfaces for this connection."), *ConnectionTargetSettings.TargetModifierName.ToString());
+		return false;
+	}
+	
+	FVCamModifierConnectionPoint* ConnectionPoint = Modifier.ConnectionPoints.Find(ConnectionPointName);
+	if (!ConnectionPoint)
+	{
+		UE_CLOG(bLogWarnings, LogVCamConnection, Warning, TEXT("Connection Failed: Failed to find Connection Point with name %s in Modifier %s"),
+			*ConnectionTargetSettings.TargetConnectionPoint.ToString(),
+			*ConnectionTargetSettings.TargetModifierName.ToString()
+			);
+		return false;
+	}
+	
+	// The Connection Point is considered valid if either we don't require an action or the required ActionType matches the ActionType inside the Connection Point
+	const bool bConnectionPointHasRequiredActionType = !bRequiresInputAction || (ConnectionPoint->AssociatedAction && ConnectionPoint->AssociatedAction->ValueType == ActionType);
+	if (!bConnectionPointHasRequiredActionType)
+	{
+		UE_CLOG(bLogWarnings, LogVCamConnection, Warning, TEXT("Connection Failed: Connection Point with name %s does not provide an action of the required type. Expected %s but found %s"),
+			*ConnectionTargetSettings.TargetConnectionPoint.ToString(),
+			*UEnum::GetValueAsString(ActionType),
+			ConnectionPoint->AssociatedAction ? *UEnum::GetValueAsString(ConnectionPoint->AssociatedAction->ValueType) : TEXT("None")
+			);
+		return false;
+	}
+	
+	return true;
 }
 
 void FVCamConnection::ResetConnection()
diff --git a/Engine/Plugins/Experimental/VirtualCamera/Source/VCamCore/Private/Util/ConnectionUtils.cpp b/Engine/Plugins/Experimental/VirtualCamera/Source/VCamCore/Private/Util/ConnectionUtils.cpp
new file mode 100644
--- /dev/null
+++ b/Engine/Plugins/Experimental/VirtualCamera/Source/VCamCore/Private/Util/ConnectionUtils.cpp
@@ -0,0 +1,66 @@
+// Copyright Epic Games, Inc. All Rights Reserved.
+
+#include "Util/ConnectionUtils.h"
+
+#include "Modifier/VCamModifier.h"
+#include "UI/VCamConnectionStructs.h"
+#include "VCamComponent.h"
+
+namespace UE::VCamCore::ConnectionUtils
+{
+	TArray<FModifierStackEntry> FindCompatibleModifiers(const FVCamConnection& Connection, UVCamComponent& Component)
+	{
+		TArray<FModifierStackEntry> Modifiers;
+		Modifiers.Reserve(Component.GetNumberOfModifiers());
+		ForEachCompatibleConnectionPoint(Connection, Component, [&Modifiers](const FModifierStackEntry& Modifier, FName ConnectionPoint)
+		{
+			Modifiers.Add(Modifier);
+			// Minor optimisation: no need to iterate the other connection points since we care only whether at least one is valid
+			return EBreakBehavior::SkipCurrentModifier;
+		});
+		return Modifiers;
+	}
+
+	TArray<FName> FindCompatibleConnectionPoints(const FVCamConnection& Connection, UVCamModifier& Modifier)
+	{
+		TArray<FName> ConnectionPoints;
+		ConnectionPoints.Reserve(Modifier.ConnectionPoints.Num());
+		ForEachCompatibleConnectionPoint(Connection, { FModifierStackEntry{ Modifier } }, [&ConnectionPoints](const FModifierStackEntry& Modifier, FName ConnectionPoint)
+		{
+			ConnectionPoints.Add(ConnectionPoint);
+			return EBreakBehavior::Continue;
+		});
+		return ConnectionPoints;
+	}
+	
+	void ForEachCompatibleConnectionPoint(const FVCamConnection& Connection, const TArray<FModifierStackEntry>& Modifiers, FCompatibleModifierCallback ProcessConnectionPoint)
+	{
+		constexpr bool bLogWarnings = false;
+		for (const FModifierStackEntry& StackEntry : Modifiers)
+		{
+			// We allow disabled modifier at this point - users of ForEachCompatibleConnectionPoint can filter it themselves
+			if (!StackEntry.GeneratedModifier) 
+			{
+				continue;
+			}
+			
+			for (TPair<FName, FVCamModifierConnectionPoint>& ConnectionPoint : StackEntry.GeneratedModifier->ConnectionPoints)
+			{
+				if (!Connection.IsConnectionValid(*StackEntry.GeneratedModifier, ConnectionPoint.Key, bLogWarnings))
+				{
+					continue;
+				}
+
+				const EBreakBehavior BreakBehavior = ProcessConnectionPoint(StackEntry, ConnectionPoint.Key);
+				if (BreakBehavior == EBreakBehavior::SkipCurrentModifier)
+				{
+					break;
+				}
+				if (BreakBehavior == EBreakBehavior::Break)
+				{
+					return;
+				}
+			}
+		}
+	}
+}
\ No newline at end of file
diff --git a/Engine/Plugins/Experimental/VirtualCamera/Source/VCamCore/Public/Modifier/ModifierStackEntry.h b/Engine/Plugins/Experimental/VirtualCamera/Source/VCamCore/Public/Modifier/ModifierStackEntry.h
--- a/Engine/Plugins/Experimental/VirtualCamera/Source/VCamCore/Public/Modifier/ModifierStackEntry.h
+++ b/Engine/Plugins/Experimental/VirtualCamera/Source/VCamCore/Public/Modifier/ModifierStackEntry.h
@@ -16,7 +16,7 @@ struct FModifierStackEntry
 
 	// Identifier for this modifier in the stack
     UPROPERTY(EditAnywhere, Category="Modifier")
-    FName Name;
+    FName Name = NAME_None;
 
 	// Controls whether the modifier actually gets applied
 	UPROPERTY(EditAnywhere, Category="Modifier")
@@ -38,6 +38,10 @@ struct FModifierStackEntry
 #endif
 	}
 
+	FModifierStackEntry(UVCamModifier& Modifier)
+		: GeneratedModifier(&Modifier)
+	{}
+
 	// If ModifierClass is provided then you must also supply a valid outer for the generated modifier
 	FModifierStackEntry(const FName& InName, const TSubclassOf<UVCamModifier> InModifierClass, UObject* InOuter)
 		: Name(InName)
diff --git a/Engine/Plugins/Experimental/VirtualCamera/Source/VCamCore/Public/UI/VCamConnectionStructs.h b/Engine/Plugins/Experimental/VirtualCamera/Source/VCamCore/Public/UI/VCamConnectionStructs.h
--- a/Engine/Plugins/Experimental/VirtualCamera/Source/VCamCore/Public/UI/VCamConnectionStructs.h
+++ b/Engine/Plugins/Experimental/VirtualCamera/Source/VCamCore/Public/UI/VCamConnectionStructs.h
@@ -125,6 +125,9 @@ struct VCAMCORE_API FVCamConnection
 	 */
 	bool AttemptConnection(UVCamComponent* VCamComponent);
 
+	/** Whether under the filter criteria of this struct it is legal to connection to ConnectionPoint from Modifier. */
+	bool IsConnectionValid(UVCamModifier& Modifier, FName ConnectionPointName, bool bLogWarnings = false) const;
+
 	/*
 	 * Clears any currently connected modifier and action
 	 */
diff --git a/Engine/Plugins/Experimental/VirtualCamera/Source/VCamCore/Public/Util/ConnectionUtils.h b/Engine/Plugins/Experimental/VirtualCamera/Source/VCamCore/Public/Util/ConnectionUtils.h
new file mode 100644
--- /dev/null
+++ b/Engine/Plugins/Experimental/VirtualCamera/Source/VCamCore/Public/Util/ConnectionUtils.h
@@ -0,0 +1,44 @@
+// Copyright Epic Games, Inc. All Rights Reserved.
+
+#pragma once
+
+#include "CoreMinimal.h"
+#include "VCamComponent.h"
+
+class UVCamModifier;
+struct FVCamConnection;
+struct FVCamModifierConnectionPoint;
+
+namespace UE::VCamCore::ConnectionUtils
+{
+	/** Gets all modifiers that satisfy the connection's filter requirements and have at least one connection point available. */
+	VCAMCORE_API TArray<FModifierStackEntry> FindCompatibleModifiers(const FVCamConnection& Connection, UVCamComponent& Component);
+
+	FORCEINLINE TArray<FName> FindCompatibleModifierNames(const FVCamConnection& Connection, UVCamComponent& Component)
+	{
+		TArray<FName> Result;
+		Algo::Transform(FindCompatibleModifiers(Connection, Component), Result, [](const FModifierStackEntry& Entry){ return Entry.Name; });
+		return Result;
+	}
+
+	/** Gets all connection points that satisfy the connection's filter requirements. */
+	VCAMCORE_API TArray<FName> FindCompatibleConnectionPoints(const FVCamConnection& Connection, UVCamModifier& Modifier);
+
+	enum class EBreakBehavior
+	{
+		/** Continue with the next connection point or modifier */
+		Continue,
+		/** Continue with the next connection point on the next modifier  */
+		SkipCurrentModifier,
+		/** Stop */
+		Break
+	};
+	using FCompatibleModifierCallback = TFunctionRef<EBreakBehavior(const FModifierStackEntry& Modifier, FName ConnectionPoint)>;
+
+	VCAMCORE_API void ForEachCompatibleConnectionPoint(const FVCamConnection& Connection, const TArray<FModifierStackEntry>& Modifiers, FCompatibleModifierCallback ProcessConnectionPoint);
+	
+	FORCEINLINE void ForEachCompatibleConnectionPoint(const FVCamConnection& Connection, UVCamComponent& Component, FCompatibleModifierCallback ProcessConnectionPoint)
+	{
+		ForEachCompatibleConnectionPoint(Connection, Component.GetModifierStack(), ProcessConnectionPoint);
+	}
+}
diff --git a/Engine/Plugins/Experimental/VirtualCamera/Source/VCamCore/Public/VCamComponent.h b/Engine/Plugins/Experimental/VirtualCamera/Source/VCamCore/Public/VCamComponent.h
--- a/Engine/Plugins/Experimental/VirtualCamera/Source/VCamCore/Public/VCamComponent.h
+++ b/Engine/Plugins/Experimental/VirtualCamera/Source/VCamCore/Public/VCamComponent.h
@@ -182,17 +182,19 @@ public:
 	 */
 	UFUNCTION(BlueprintPure, Category = "VirtualCamera")
 	void GetAllModifiers(TArray<UVCamModifier*>& Modifiers) const;
+	
 	/** Returns all the modifier names used to identifying connection points. */
 	UFUNCTION(BlueprintPure, Category = "VirtualCamera")
 	TArray<FName> GetAllModifierNames() const;
+
+	const TArray<FModifierStackEntry>& GetModifierStack() const { return ModifierStack; }
+
 	/** Returns the Modifier in the Stack with the given index if it exist. */
 	UFUNCTION(BlueprintPure, Category = "VirtualCamera")
 	UVCamModifier* GetModifierByIndex(const int32 Index) const;
-	/**
-	 * Tries to find a Modifier in the Stack with the given name. The returned Modifier must be checked before it is used. */
+	/** Tries to find a Modifier in the Stack with the given name. The returned Modifier must be checked before it is used. */
 	UFUNCTION(BlueprintPure, Category = "VirtualCamera")
 	UVCamModifier* GetModifierByName(const FName Name) const;
-
 	/** Given a specific Modifier class, returns a list of matching Modifiers */
 	UFUNCTION(BlueprintPure, Category = "VirtualCamera", Meta = (DeterminesOutputType = "ModifierClass", DynamicOutputParam = "FoundModifiers"))
 	void GetModifiersByClass(UPARAM(meta = (AllowAbstract = "false")) TSubclassOf<UVCamModifier> ModifierClass, TArray<UVCamModifier*>& FoundModifiers) const;
diff --git a/Engine/Plugins/Experimental/VirtualCamera/Source/VCamCoreEditor/Private/Customizations/ConnectionTargetSettingsTypeCustomization.cpp b/Engine/Plugins/Experimental/VirtualCamera/Source/VCamCoreEditor/Private/Customizations/ConnectionTargetSettingsTypeCustomization.cpp
--- a/Engine/Plugins/Experimental/VirtualCamera/Source/VCamCoreEditor/Private/Customizations/ConnectionTargetSettingsTypeCustomization.cpp
+++ b/Engine/Plugins/Experimental/VirtualCamera/Source/VCamCoreEditor/Private/Customizations/ConnectionTargetSettingsTypeCustomization.cpp
@@ -13,6 +13,7 @@
 #include "Modules/ModuleManager.h"
 #include "Selection.h"
 #include "SSimpleComboButton.h"
+#include "Util/ConnectionUtils.h"
 #include "Widgets/Images/SImage.h"
 #include "Widgets/Layout/SBox.h"
 #include "Widgets/SBoxPanel.h"
@@ -48,12 +49,13 @@ namespace UE::VCamCoreEditor::Private
 		IDetailChildrenBuilder& ChildBuilder,
 		IPropertyTypeCustomizationUtils& CustomizationUtils)
 	{
-		uint32 NumChildren;
-		PropertyHandle->GetNumChildren( NumChildren );	
 		const FName WidgetProperty = GET_MEMBER_NAME_CHECKED(FVCamConnectionTargetSettings, TargetModifierName);
 		const FName ConnectionTargetsProperty = GET_MEMBER_NAME_CHECKED(FVCamConnectionTargetSettings, TargetConnectionPoint);
 		TSharedPtr<IPropertyHandle> TargetModifierNameProperty;
 		TSharedPtr<IPropertyHandle> TargetConnectionPointProperty;
+		
+		uint32 NumChildren;
+		PropertyHandle->GetNumChildren( NumChildren);	
 		for (uint32 ChildIndex = 0; ChildIndex < NumChildren; ++ChildIndex)
 		{
 			TSharedPtr<IPropertyHandle> ChildProperty = PropertyHandle->GetChildHandle(ChildIndex);
@@ -69,11 +71,19 @@ namespace UE::VCamCoreEditor::Private
 		
 		AddScopeRow(ChildBuilder, CustomizationUtils);
 		
+		const TSharedPtr<IPropertyHandle> ParentHandle = PropertyHandle->GetParentHandle();
+		const bool bIsValidHandle = ParentHandle
+			|| !CastField<FStructProperty>(PropertyHandle->GetProperty())
+			|| CastField<FStructProperty>(PropertyHandle->GetProperty())->Struct == FVCamConnection::StaticStruct();
+		const TSharedPtr<IPropertyHandle> OptionalVCamConnectionParentStructHandle = bIsValidHandle
+			? ParentHandle
+			: nullptr;
+		
 		IDetailPropertyRow& ModifierRow = ChildBuilder.AddProperty(TargetModifierNameProperty.ToSharedRef());
-		CustomizeModifier(TargetModifierNameProperty.ToSharedRef(), ChildBuilder, CustomizationUtils, ModifierRow);
+		CustomizeModifier(TargetModifierNameProperty.ToSharedRef(), ModifierRow, OptionalVCamConnectionParentStructHandle);
 
 		IDetailPropertyRow& TargetConnectionRow = ChildBuilder.AddProperty(TargetModifierNameProperty.ToSharedRef());
-		CustomizeConnectionPoint(TargetModifierNameProperty.ToSharedRef(), TargetConnectionPointProperty.ToSharedRef(), ChildBuilder, CustomizationUtils, TargetConnectionRow);
+		CustomizeConnectionPoint(TargetModifierNameProperty.ToSharedRef(), TargetConnectionPointProperty.ToSharedRef(), TargetConnectionRow, OptionalVCamConnectionParentStructHandle);
 	}
 
 	void FConnectionTargetSettingsTypeCustomization::AddScopeRow(IDetailChildrenBuilder& ChildBuilder, IPropertyTypeCustomizationUtils& CustomizationUtils)
@@ -112,18 +122,31 @@ namespace UE::VCamCoreEditor::Private
 
 	void FConnectionTargetSettingsTypeCustomization::CustomizeModifier(
 		TSharedRef<IPropertyHandle> ModifierHandle,
-		IDetailChildrenBuilder& ChildBuilder,
-		IPropertyTypeCustomizationUtils& CustomizationUtils,
-		IDetailPropertyRow& Row
+		IDetailPropertyRow& Row,
+		TSharedPtr<IPropertyHandle> OptionalVCamConnectionParentStructHandle // Can be null
 		) const
 	{
 		CustomizeNameProperty(ModifierHandle, Row,
-			TAttribute<TArray<FName>>::CreateLambda([]() -> TArray<FName>
+			TAttribute<TArray<FName>>::CreateLambda([OptionalVCamConnectionParentStructHandle]() -> TArray<FName>
 			{
 				UVCamComponent* DataSource = GetUserFocusedConnectionPointSource().Component.Get();
-				return DataSource
-					? DataSource->GetAllModifierNames()
-					: TArray<FName>{};
+				if (!DataSource)
+				{
+					return {};
+				}
+
+				// If the property is within a FVCamConnection, narrow down the list of suggested modifiers
+				void* ValueData = nullptr;
+				if (OptionalVCamConnectionParentStructHandle
+					&& OptionalVCamConnectionParentStructHandle->IsValidHandle()
+					&& OptionalVCamConnectionParentStructHandle->GetValueData(ValueData) == FPropertyAccess::Success
+					&& ValueData)
+				{
+					const FVCamConnection* ConnectionData = static_cast<FVCamConnection*>(ValueData);
+					return VCamCore::ConnectionUtils::FindCompatibleModifierNames(*ConnectionData, *DataSource);
+				}
+				
+				return DataSource->GetAllModifierNames();
 			}),
 			TAttribute<bool>::CreateLambda([](){ return GetUserFocusedConnectionPointSource().ComponentSource != EComponentSource::None; })
 			);
@@ -132,13 +155,12 @@ namespace UE::VCamCoreEditor::Private
 	void FConnectionTargetSettingsTypeCustomization::CustomizeConnectionPoint(
 		TSharedRef<IPropertyHandle> ModifierHandle,
 		TSharedRef<IPropertyHandle> ConnectionPointHandle,
-		IDetailChildrenBuilder& ChildBuilder,
-		IPropertyTypeCustomizationUtils& CustomizationUtils,
-		IDetailPropertyRow& Row
+		IDetailPropertyRow& Row,
+		TSharedPtr<IPropertyHandle> OptionalVCamConnectionParentStructHandle // Can be null
 		) const
 	{
 		CustomizeNameProperty(ConnectionPointHandle, Row,
-			TAttribute<TArray<FName>>::CreateLambda([ModifierHandle]() -> TArray<FName>
+			TAttribute<TArray<FName>>::CreateLambda([ModifierHandle, OptionalVCamConnectionParentStructHandle]() -> TArray<FName>
 			{
 				FName ModifierName;
 				if (ModifierHandle->GetValue(ModifierName) != FPropertyAccess::Success)
@@ -147,13 +169,30 @@ namespace UE::VCamCoreEditor::Private
 				}
 
 				UVCamComponent* DataSource = GetUserFocusedConnectionPointSource().Component.Get();
+				if (!DataSource)
+				{
+					return {};
+				}
+				
 				UVCamModifier* Modifier = DataSource->GetModifierByName(ModifierName);
+				if (!Modifier)
+				{
+					return {};
+				}
 
-				TArray<FName> ConnectionPoints;
-				if (Modifier)
+				// If the property is within a FVCamConnection, narrow down the list of suggested modifiers
+				void* ValueData = nullptr;
+				if (OptionalVCamConnectionParentStructHandle
+					&& OptionalVCamConnectionParentStructHandle->IsValidHandle()
+					&& OptionalVCamConnectionParentStructHandle->GetValueData(ValueData) == FPropertyAccess::Success
+					&& ValueData)
 				{
-					Modifier->ConnectionPoints.GenerateKeyArray(ConnectionPoints);
+					const FVCamConnection* ConnectionData = static_cast<FVCamConnection*>(ValueData);
+					return VCamCore::ConnectionUtils::FindCompatibleConnectionPoints(*ConnectionData, *Modifier);
 				}
+
+				TArray<FName> ConnectionPoints;
+				Modifier->ConnectionPoints.GenerateKeyArray(ConnectionPoints);
 				return ConnectionPoints;
 			}),
 			TAttribute<bool>::CreateLambda([ModifierHandle]()
@@ -193,18 +232,20 @@ namespace UE::VCamCoreEditor::Private
 		}
 
 		// Level Editor
+		USelection* Selection = GEditor->GetSelectedActors();
+		if (Selection)
 		{
-			USelection* Selection = GEditor->GetSelectedActors();
-			UObject* SelectedObject = Selection && Selection->Num() == 1
-				? Selection->GetSelectedObject(0)
-				: nullptr;
-			AActor* SelectedActor = Cast<AActor>(SelectedObject);
-			UVCamComponent* VCamComponent = SelectedActor
-				? SelectedActor->FindComponentByClass<UVCamComponent>()
-				: nullptr;
-			if (IsValid(VCamComponent))
+			for (int32 i = 0; i < Selection->Num(); ++i)
 			{
-				return { EComponentSource::LevelSelection, VCamComponent };
+				UObject* SelectedObject = Selection->GetSelectedObject(i);
+				AActor* SelectedActor = Cast<AActor>(SelectedObject);
+				UVCamComponent* VCamComponent = SelectedActor
+					? SelectedActor->FindComponentByClass<UVCamComponent>()
+					: nullptr;
+				if (IsValid(VCamComponent))
+				{
+					return { EComponentSource::LevelSelection, VCamComponent };
+				}
 			}
 		}
 			
diff --git a/Engine/Plugins/Experimental/VirtualCamera/Source/VCamCoreEditor/Private/Customizations/ConnectionTargetSettingsTypeCustomization.h b/Engine/Plugins/Experimental/VirtualCamera/Source/VCamCoreEditor/Private/Customizations/ConnectionTargetSettingsTypeCustomization.h
--- a/Engine/Plugins/Experimental/VirtualCamera/Source/VCamCoreEditor/Private/Customizations/ConnectionTargetSettingsTypeCustomization.h
+++ b/Engine/Plugins/Experimental/VirtualCamera/Source/VCamCoreEditor/Private/Customizations/ConnectionTargetSettingsTypeCustomization.h
@@ -29,8 +29,8 @@ namespace UE::VCamCoreEditor::Private
 		
 		void AddScopeRow(IDetailChildrenBuilder& ChildBuilder, IPropertyTypeCustomizationUtils& CustomizationUtils);
 
-		void CustomizeModifier(TSharedRef<IPropertyHandle> ModifierHandle, IDetailChildrenBuilder& ChildBuilder, IPropertyTypeCustomizationUtils& CustomizationUtils, IDetailPropertyRow& Row) const;
-		void CustomizeConnectionPoint(TSharedRef<IPropertyHandle> ModifierHandle, TSharedRef<IPropertyHandle> ConnectionPointHandle, IDetailChildrenBuilder& ChildBuilder, IPropertyTypeCustomizationUtils& CustomizationUtils, IDetailPropertyRow& Row) const;
+		void CustomizeModifier(TSharedRef<IPropertyHandle> ModifierHandle, IDetailPropertyRow& Row, TSharedPtr<IPropertyHandle> OptionalVCamConnectionParentStructHandle) const;
+		void CustomizeConnectionPoint(TSharedRef<IPropertyHandle> ModifierHandle, TSharedRef<IPropertyHandle> ConnectionPointHandle, IDetailPropertyRow& Row, TSharedPtr<IPropertyHandle> OptionalVCamConnectionParentStructHandle) const;
 
 		enum class EComponentSource
 		{
diff --git a/Engine/Plugins/Experimental/VirtualCamera/Source/VCamCoreEditor/Private/Customizations/OutputProvider/ConnectionRemapUtilsImpl.cpp b/Engine/Plugins/Experimental/VirtualCamera/Source/VCamCoreEditor/Private/Customizations/OutputProvider/ConnectionRemapUtilsImpl.cpp
--- a/Engine/Plugins/Experimental/VirtualCamera/Source/VCamCoreEditor/Private/Customizations/OutputProvider/ConnectionRemapUtilsImpl.cpp
+++ b/Engine/Plugins/Experimental/VirtualCamera/Source/VCamCoreEditor/Private/Customizations/OutputProvider/ConnectionRemapUtilsImpl.cpp
@@ -22,14 +22,14 @@ namespace UE::VCamCoreEditor::Private
 		}
 		
 		TSharedPtr<FStructOnScope> StructData;
-		if (const TSharedPtr<FStructOnScope>* ExistingStructData = AddedConnections.Find(Args.PropertyName))
+		if (const TSharedPtr<FStructOnScope>* ExistingStructData = AddedConnections.Find(Args.ConnectionName))
 		{
 			StructData = *ExistingStructData;
 		}
 		else
 		{
 			StructData = MakeShared<FStructOnScope>(MoveTemp(Args.StructData));
-			AddedConnections.Add(Args.PropertyName, StructData);
+			AddedConnections.Add(Args.ConnectionName, StructData);
 		}
 		
 		const TSharedPtr<IPropertyHandle> PropertyHandle = BuilderPin->AddStructurePropertyData(StructData, Args.PropertyName);
