Drop-downs for connection target settings do not respect RequiredInterfaces: Add support for remapping widget connections in editor

diff --git a/Engine/Plugins/Experimental/VirtualCamera/Source/VCamCoreEditor/Private/Customizations/OutputProvider/ConnectionContainerDummy.h b/Engine/Plugins/Experimental/VirtualCamera/Source/VCamCoreEditor/Private/Customizations/OutputProvider/ConnectionContainerDummy.h
new file mode 100644
--- /dev/null
+++ b/Engine/Plugins/Experimental/VirtualCamera/Source/VCamCoreEditor/Private/Customizations/OutputProvider/ConnectionContainerDummy.h
@@ -0,0 +1,23 @@
+// Copyright Epic Games, Inc. All Rights Reserved.
+
+#pragma once
+
+#include "CoreMinimal.h"
+#include "UI/VCamConnectionStructs.h"
+#include "ConnectionContainerDummy.generated.h"
+
+/**
+ * This dummy is needed so that the FConnectionTargetSettingsTypeCustomization can access FVCamConnection and use
+ * RequiredInterfaces, etc. for filtering suggestions.
+ * 
+ * The detail customization systems implementation requires a valid parent handle to whatever contains the FVCamConnection,
+ * which is FConnectionContainerDummy's task.
+ */
+USTRUCT()
+struct FConnectionContainerDummy
+{
+	GENERATED_BODY()
+
+	UPROPERTY(EditAnywhere, Category = "Dummy")
+	FVCamConnection Connection;
+};
\ No newline at end of file
diff --git a/Engine/Plugins/Experimental/VirtualCamera/Source/VCamCoreEditor/Private/Customizations/OutputProvider/ConnectionRemapCustomization_VCamWidget.cpp b/Engine/Plugins/Experimental/VirtualCamera/Source/VCamCoreEditor/Private/Customizations/OutputProvider/ConnectionRemapCustomization_VCamWidget.cpp
--- a/Engine/Plugins/Experimental/VirtualCamera/Source/VCamCoreEditor/Private/Customizations/OutputProvider/ConnectionRemapCustomization_VCamWidget.cpp
+++ b/Engine/Plugins/Experimental/VirtualCamera/Source/VCamCoreEditor/Private/Customizations/OutputProvider/ConnectionRemapCustomization_VCamWidget.cpp
@@ -39,18 +39,11 @@ namespace UE::VCamCoreEditor::Private
 				continue;
 			}
 			
-			// Overwrite bManuallyConfigureConnection because otherwise its EditCondition will hide the property!
-			FVCamConnection ConnectionCopy = Connection.Value;
-			ConnectionCopy.bManuallyConfigureConnection = true;
-			TStructOnScope<FVCamConnection> StructOnScope(ConnectionCopy);
-			
 			const FName ConnectionName = Connection.Key;
 			Args.Utils->AddConnection(FAddConnectionArgs{
 				Args.WidgetGroup,
 				ConnectionName,
-				GET_MEMBER_NAME_CHECKED(FVCamConnection, ConnectionTargetSettings),
-				// Object slicing occurs here
-				static_cast<FStructOnScope>(MoveTemp(StructOnScope)),
+				Connection.Value,
 				FOnTargetSettingsChanged::CreateSP(this, &FConnectionRemapCustomization_VCamWidget::OnTargetSettingsChanged, ConnectionName),
 				Args.Utils->GetRegularFont()
 			});
diff --git a/Engine/Plugins/Experimental/VirtualCamera/Source/VCamCoreEditor/Private/Customizations/OutputProvider/ConnectionRemapUtilsImpl.cpp b/Engine/Plugins/Experimental/VirtualCamera/Source/VCamCoreEditor/Private/Customizations/OutputProvider/ConnectionRemapUtilsImpl.cpp
--- a/Engine/Plugins/Experimental/VirtualCamera/Source/VCamCoreEditor/Private/Customizations/OutputProvider/ConnectionRemapUtilsImpl.cpp
+++ b/Engine/Plugins/Experimental/VirtualCamera/Source/VCamCoreEditor/Private/Customizations/OutputProvider/ConnectionRemapUtilsImpl.cpp
@@ -21,33 +21,45 @@ namespace UE::VCamCoreEditor::Private
 			return;
 		}
 		
-		TSharedPtr<FStructOnScope> StructData;
-		if (const TSharedPtr<FStructOnScope>* ExistingStructData = AddedConnections.Find(Args.ConnectionName))
+		TSharedPtr<TStructOnScope<FConnectionContainerDummy>> StructData;
+		if (const TSharedPtr<TStructOnScope<FConnectionContainerDummy>>* ExistingStructData = AddedConnections.Find(Args.ConnectionName))
 		{
 			StructData = *ExistingStructData;
 		}
 		else
 		{
-			StructData = MakeShared<FStructOnScope>(MoveTemp(Args.StructData));
+			FConnectionContainerDummy DummyContainer { Args.ConnectionData };
+			// Overwrite bManuallyConfigureConnection because otherwise its EditCondition will hide the property!
+			DummyContainer.Connection.bManuallyConfigureConnection = true;
+			TStructOnScope<FConnectionContainerDummy> StructOnScope(DummyContainer);
+
+			// Object slicing occurs intentionally
+			StructData = MakeShared<TStructOnScope<FConnectionContainerDummy>>(MoveTemp(StructOnScope));
 			AddedConnections.Add(Args.ConnectionName, StructData);
 		}
-		
-		const TSharedPtr<IPropertyHandle> PropertyHandle = BuilderPin->AddStructurePropertyData(StructData, Args.PropertyName);
-		if (!ensureMsgf(PropertyHandle, TEXT("No property of %s found on struct"), *Args.PropertyName.ToString()))
+
+		const FName TargetSettingsPropertyName = GET_MEMBER_NAME_CHECKED(FConnectionContainerDummy, Connection);
+		const TSharedPtr<IPropertyHandle> ConnectionHandle = BuilderPin->AddStructurePropertyData(StructData, TargetSettingsPropertyName);
+		if (!ensure(ConnectionHandle))
+		{
+			return;
+		}
+		const TSharedPtr<IPropertyHandle> TargetConnectionSettings = ConnectionHandle->GetChildHandle(GET_MEMBER_NAME_CHECKED(FVCamConnection, ConnectionTargetSettings));
+		if (!ensure(TargetConnectionSettings))
 		{
 			return;
 		}
 		
-		PropertyHandle->SetOnChildPropertyValueChanged(FSimpleDelegate::CreateLambda([PropertyHandle, Callback = MoveTemp(Args.OnTargetSettingsChangedDelegate)]()
+		TargetConnectionSettings->SetOnChildPropertyValueChanged(FSimpleDelegate::CreateLambda([TargetConnectionSettings, Callback = MoveTemp(Args.OnTargetSettingsChangedDelegate)]()
 		{
 			void* Data;
-			if (ensure(PropertyHandle->GetValueData(Data) == FPropertyAccess::Success))
+			if (ensure(TargetConnectionSettings->GetValueData(Data) == FPropertyAccess::Success))
 			{
 				Callback.Execute(*(FVCamConnectionTargetSettings*)Data);
 			}
 		}));
 
-		Args.DetailGroup.AddPropertyRow(PropertyHandle.ToSharedRef())
+		Args.DetailGroup.AddPropertyRow(TargetConnectionSettings.ToSharedRef())
 			.DisplayName(FText::FromName(Args.ConnectionName)) 
 			.CustomWidget(true)
 			.NameContent()
@@ -58,7 +70,7 @@ namespace UE::VCamCoreEditor::Private
 			]
 			.ValueContent()
 			[
-				PropertyHandle->CreatePropertyValueWidget()
+				TargetConnectionSettings->CreatePropertyValueWidget()
 			];
 	}
 
diff --git a/Engine/Plugins/Experimental/VirtualCamera/Source/VCamCoreEditor/Private/Customizations/OutputProvider/ConnectionRemapUtilsImpl.h b/Engine/Plugins/Experimental/VirtualCamera/Source/VCamCoreEditor/Private/Customizations/OutputProvider/ConnectionRemapUtilsImpl.h
--- a/Engine/Plugins/Experimental/VirtualCamera/Source/VCamCoreEditor/Private/Customizations/OutputProvider/ConnectionRemapUtilsImpl.h
+++ b/Engine/Plugins/Experimental/VirtualCamera/Source/VCamCoreEditor/Private/Customizations/OutputProvider/ConnectionRemapUtilsImpl.h
@@ -3,6 +3,7 @@
 #pragma once
 
 #include "CoreMinimal.h"
+#include "ConnectionContainerDummy.h"
 #include "DetailLayoutBuilder.h"
 #include "Customization/IConnectionRemapUtils.h"
 
@@ -23,7 +24,7 @@ namespace UE::VCamCoreEditor::Private
 	private:
 		
 		TWeakPtr<IDetailLayoutBuilder> Builder;
-		TMap<FName, TSharedPtr<FStructOnScope>> AddedConnections;
+		TMap<FName, TSharedPtr<TStructOnScope<FConnectionContainerDummy>>> AddedConnections;
 	};
 
 }
diff --git a/Engine/Plugins/Experimental/VirtualCamera/Source/VCamCoreEditor/Public/Customization/IConnectionRemapUtils.h b/Engine/Plugins/Experimental/VirtualCamera/Source/VCamCoreEditor/Public/Customization/IConnectionRemapUtils.h
--- a/Engine/Plugins/Experimental/VirtualCamera/Source/VCamCoreEditor/Public/Customization/IConnectionRemapUtils.h
+++ b/Engine/Plugins/Experimental/VirtualCamera/Source/VCamCoreEditor/Public/Customization/IConnectionRemapUtils.h
@@ -12,16 +12,15 @@ struct FVCamConnectionTargetSettings;
 namespace UE::VCamCoreEditor
 {
 	DECLARE_DELEGATE_OneParam(FOnTargetSettingsChanged, const FVCamConnectionTargetSettings& NewSettings);
+	/** Data for requesting to add connection target settings to details panel */
 	struct FAddConnectionArgs
 	{
 		IDetailGroup& DetailGroup;
 
+		/** Unique ID used for mapping data */
 		FName ConnectionName;
-
-		/** Name of the property.  */
-		FName PropertyName;
-		/** A containing a FVCamConnectionTargetSettings property with name PropertyName */
-		FStructOnScope StructData;
+		/** Contains the data to display */
+		const FVCamConnection& ConnectionData;
 
 		/** Called when the settings are changed. Copy the passed settings to your UPROPERTY(). */
 		FOnTargetSettingsChanged OnTargetSettingsChangedDelegate;
@@ -29,11 +28,10 @@ namespace UE::VCamCoreEditor
 		/** The font to use for displaying property texts. */
 		FSlateFontInfo Font;
 
-		FAddConnectionArgs(IDetailGroup& DetailGroup, FName ConnectionName, FName PropertyName, FStructOnScope StructData, FOnTargetSettingsChanged OnTargetSettingsChangedDelegate, FSlateFontInfo Font)
+		FAddConnectionArgs(IDetailGroup& DetailGroup, FName ConnectionName, const FVCamConnection& ConnectionData, FOnTargetSettingsChanged OnTargetSettingsChangedDelegate, FSlateFontInfo Font)
 			: DetailGroup(DetailGroup)
 			, ConnectionName(ConnectionName)
-			, PropertyName(PropertyName)
-			, StructData(MoveTemp(StructData))
+			, ConnectionData(ConnectionData)
 			, OnTargetSettingsChangedDelegate(MoveTemp(OnTargetSettingsChangedDelegate))
 			, Font(MoveTemp(Font))
 		{}
@@ -44,7 +42,7 @@ namespace UE::VCamCoreEditor
 	{
 	public:
 
-		/** Adds a property row representing TargetSettings to DetailGroup. */
+		/** Adds a property row representing FVCamConnection::TargetSettings to DetailGroup. */
 		virtual void AddConnection(FAddConnectionArgs Params) = 0;
 
 		/** @return the font used for properties and details */
