Add comments to better explain intent

diff --git a/Engine/Plugins/Experimental/VirtualCamera/Source/VCamCoreEditor/Private/Customizations/OutputProvider/OutputProviderTypeCustomization.cpp b/Engine/Plugins/Experimental/VirtualCamera/Source/VCamCoreEditor/Private/Customizations/OutputProvider/OutputProviderTypeCustomization.cpp
--- a/Engine/Plugins/Experimental/VirtualCamera/Source/VCamCoreEditor/Private/Customizations/OutputProvider/OutputProviderTypeCustomization.cpp
+++ b/Engine/Plugins/Experimental/VirtualCamera/Source/VCamCoreEditor/Private/Customizations/OutputProvider/OutputProviderTypeCustomization.cpp
@@ -31,6 +31,7 @@ namespace UE::VCamCoreEditor::Private
 			: nullptr;
 		if (!CustomizedOutputProvider)
 		{
+			// If user just clicked + Add button, reference will be null. Let the user select a class.
 			ChildBuilder.AddCustomRow(FText::GetEmpty())
 				.NameContent()
 				[
@@ -49,7 +50,8 @@ namespace UE::VCamCoreEditor::Private
 		{
 			return;
 		}
-		
+
+		// AddExternalObjects causes FOutputProviderLayoutCustomization to customize the details.
 		IDetailPropertyRow* DetailRow = ChildBuilder.AddExternalObjects({ CustomizedOutputProvider },
 			FAddPropertyParams()
 				.CreateCategoryNodes(false)		// Avoid creating intermediate group expansion
