Persist connections remapped via details customization

diff --git a/Engine/Plugins/Experimental/VirtualCamera/Source/VCamCore/Private/Output/VCamOutputProviderBase.cpp b/Engine/Plugins/Experimental/VirtualCamera/Source/VCamCore/Private/Output/VCamOutputProviderBase.cpp
--- a/Engine/Plugins/Experimental/VirtualCamera/Source/VCamCore/Private/Output/VCamOutputProviderBase.cpp
+++ b/Engine/Plugins/Experimental/VirtualCamera/Source/VCamCore/Private/Output/VCamOutputProviderBase.cpp
@@ -14,6 +14,8 @@
 #include "Framework/Application/SlateApplication.h"
 #include "Slate/SceneViewport.h"
 #include "UObject/UObjectBaseUtility.h"
+#include "Util/ConnectionUtils.h"
+#include "Util/WidgetSnapshotUtils.h"
 
 #if WITH_EDITOR
 #include "Editor.h"
@@ -231,20 +233,20 @@ void UVCamOutputProviderBase::DisplayUMG()
 
 		if (ActorWorld)
 		{
-#if WITH_EDITOR
-			if (SavedConnectionRemappingData && SavedConnectionRemappingData->GetClass() == UMGClass)
+			UMGWidget->Display(ActorWorld);
+			if (UUserWidget* Subwidget = UMGWidget->GetWidget(); ensure(Subwidget) && WidgetSnapshot.HasData())
 			{
-				UMGWidget->SetOverrideWidget(SavedConnectionRemappingData);
+				UE::VCamCore::WidgetSnapshotUtils::Private::ApplyTreeHierarchySnapshot(WidgetSnapshot, *Subwidget);
+				// Note that NotifyWidgetOfComponentChange will cause InitializeConnections to be called - this is important for the connections to get applied!
 			}
-			UMGWidget->Display(ActorWorld);
-			SavedConnectionRemappingData = UMGWidget->GetWidget();
-#else
-			UMGWidget->Display(ActorWorld);
-#endif
 			UE_LOG(LogVCamOutputProvider, Log, TEXT("DisplayUMG widget displayed in WorldType %d"), WorldType);
 		}
 
 		NotifyWidgetOfComponentChange();
+#if WITH_EDITOR
+		// Start registering after the initial calls to InitializeConnections to prevent unneeded snapshotting.
+		StartDetectAndSnapshotWhenConnectionsChange();
+#endif
 	}
 }
 
@@ -254,6 +256,16 @@ void UVCamOutputProviderBase::DestroyUMG()
 	{
 		if (UMGWidget->IsDisplayed())
 		{
+#if WITH_EDITOR
+			// The state only needs to be saved in the editor
+			if (UUserWidget* Subwidget = UMGWidget->GetWidget(); ensure(Subwidget))
+			{
+				StopDetectAndSnapshotWhenConnectionsChange();
+				Modify();
+				WidgetSnapshot = UE::VCamCore::WidgetSnapshotUtils::Private::TakeTreeHierarchySnapshot(*Subwidget);
+			}
+#endif
+			
 			UMGWidget->Hide();
 			UE_LOG(LogVCamOutputProvider, Log, TEXT("DestroyUMG widget %s hidden"), *UMGWidget->GetName());
 		}
@@ -503,6 +515,7 @@ void UVCamOutputProviderBase::PostEditChangeProperty(FPropertyChangedEvent& Prop
 		}
 		else if (PropertyName == NAME_UMGClass)
 		{
+			WidgetSnapshot.Reset();
 			if (bIsActive)
 			{
 				SetActive(false);
@@ -531,4 +544,54 @@ void UVCamOutputProviderBase::PostEditChangeProperty(FPropertyChangedEvent& Prop
 
 	Super::PostEditChangeProperty(PropertyChangedEvent);
 }
-#endif
+
+void UVCamOutputProviderBase::StartDetectAndSnapshotWhenConnectionsChange()
+{
+	check(UMGWidget);
+	UUserWidget* Widget = UMGWidget->GetWidget();
+	check(Widget);
+	
+	Widget->WidgetTree->ForEachWidget([this](UWidget* Widget)
+	{
+		if (UVCamWidget* VCamWidget = Cast<UVCamWidget>(Widget))
+		{
+			const TWeakObjectPtr<UVCamWidget> WeakWidget = VCamWidget;
+			VCamWidget->OnPostConnectionsReinitializedDelegate.AddUObject(this, &UVCamOutputProviderBase::OnConnectionReinitialized, WeakWidget);
+		}
+	});
+}
+
+void UVCamOutputProviderBase::StopDetectAndSnapshotWhenConnectionsChange()
+{
+	check(UMGWidget);
+	UUserWidget* Widget = UMGWidget->GetWidget();
+	check(Widget);
+	
+	Widget->WidgetTree->ForEachWidget([this](UWidget* Widget)
+	{
+		if (UVCamWidget* VCamWidget = Cast<UVCamWidget>(Widget))
+		{
+			const TWeakObjectPtr<UVCamWidget> WeakWidget = VCamWidget;
+			VCamWidget->OnPostConnectionsReinitializedDelegate.RemoveAll(this);
+		}
+	});
+}
+
+void UVCamOutputProviderBase::OnConnectionReinitialized(TWeakObjectPtr<UVCamWidget> Widget)
+{
+	if (Widget.IsValid())
+	{
+		if (WidgetSnapshot.HasData())
+		{
+			Modify();
+			UE::VCamCore::WidgetSnapshotUtils::Private::RetakeSnapshotForWidgetInHierarchy(WidgetSnapshot, *Widget.Get());
+		}
+		else if (UMGWidget && ensure(UMGWidget->GetWidget()))
+		{
+			Modify();
+			WidgetSnapshot = UE::VCamCore::WidgetSnapshotUtils::Private::TakeTreeHierarchySnapshot(*UMGWidget->GetWidget());
+		}
+	}
+}
+
+#endif
\ No newline at end of file
diff --git a/Engine/Plugins/Experimental/VirtualCamera/Source/VCamCore/Private/UI/VCamWidget.cpp b/Engine/Plugins/Experimental/VirtualCamera/Source/VCamCore/Private/UI/VCamWidget.cpp
--- a/Engine/Plugins/Experimental/VirtualCamera/Source/VCamCore/Private/UI/VCamWidget.cpp
+++ b/Engine/Plugins/Experimental/VirtualCamera/Source/VCamCore/Private/UI/VCamWidget.cpp
@@ -2,6 +2,8 @@
 
 #include "UI/VCamWidget.h"
 
+#include "VCamComponent.h"
+
 UVCamWidget::UVCamWidget(const FObjectInitializer& ObjectInitializer)
 	: Super(ObjectInitializer)
 {
diff --git a/Engine/Plugins/Experimental/VirtualCamera/Source/VCamCore/Private/Util/WidgetSnapshotUtils.cpp b/Engine/Plugins/Experimental/VirtualCamera/Source/VCamCore/Private/Util/WidgetSnapshotUtils.cpp
new file mode 100644
--- /dev/null
+++ b/Engine/Plugins/Experimental/VirtualCamera/Source/VCamCore/Private/Util/WidgetSnapshotUtils.cpp
@@ -0,0 +1,188 @@
+// Copyright Epic Games, Inc. All Rights Reserved.
+
+#include "WidgetSnapshotUtils.h"
+
+#include "LogVCamCore.h"
+#include "Blueprint/WidgetTree.h"
+#include "Serialization/ArchiveSerializedPropertyChain.h"
+#include "Serialization/BufferArchive.h"
+#include "Serialization/ObjectAndNameAsStringProxyArchive.h"
+#include "Serialization/ObjectReader.h"
+#include "Serialization/ObjectWriter.h"
+#include "UI/WidgetSnapshots.h"
+
+namespace UE::VCamCore::WidgetSnapshotUtils::Private
+{
+	bool FWidgetSnapshotSettings::IsClassAllowed(TSubclassOf<UWidget> Widget) const
+	{
+		for (UClass* Current = Widget; Current && Current != UWidget::StaticClass(); Current = Current->GetSuperClass())
+		{
+			if (AllowedWidgetClasses.Contains(Current))
+			{
+				return true;
+			}
+		}
+		return false;
+	}
+
+	FWidgetTreeSnapshot TakeTreeHierarchySnapshot(UUserWidget& Widget, const FWidgetSnapshotSettings& SnapshotSettings)
+	{
+		if (!Widget.WidgetTree)
+		{
+			return {};
+		}
+
+		FWidgetTreeSnapshot Result;
+		Widget.WidgetTree->ForEachWidget([&Result, &SnapshotSettings](UWidget* Widget)
+		{
+			if (Widget && SnapshotSettings.IsClassAllowed(Widget->GetClass()))
+			{
+				Result.WidgetSnapshots.Add(Widget->GetFName(), TakeWidgetSnapshot(*Widget, SnapshotSettings.AllowedProperties));
+			}
+		});
+
+		// Root widget may get re-instanced in the future which would change its name
+		const FName RootWidgetName = Widget.GetFName();
+		if (Result.WidgetSnapshots.Contains(RootWidgetName))
+		{
+			Result.RootWidget = RootWidgetName;
+		}
+			
+		return Result;
+	}
+	
+	bool ApplyTreeHierarchySnapshot(const FWidgetTreeSnapshot& Snapshot, UUserWidget& Widget)
+	{
+		if (!Widget.WidgetTree)
+		{
+			return false;
+		}
+
+		Widget.WidgetTree->ForEachWidget([&Snapshot](UWidget* Widget)
+		{
+			if (const FWidgetSnapshot* WidgetSnapshot = Snapshot.WidgetSnapshots.Find(Widget->GetFName()))
+			{
+				ApplyWidgetSnapshot(*WidgetSnapshot, *Widget);
+			}
+		});
+
+		// Root widget may have been re-instanced thus replacing its name
+		const FName RootWidgetName = Widget.GetFName();
+		if (!Snapshot.WidgetSnapshots.Contains(RootWidgetName))
+		{
+			const FWidgetSnapshot* RootWidgetSnapshot = Snapshot.WidgetSnapshots.Find(Snapshot.RootWidget);
+			if (RootWidgetSnapshot && RootWidgetSnapshot->WidgetClass == Widget.GetClass())
+			{
+				ApplyWidgetSnapshot(*RootWidgetSnapshot, Widget);
+			}
+			else
+			{
+				UE_LOG(LogVCamCore, Warning, TEXT("Failed to apply snapshot to root widget"));
+			}
+		}
+		
+		return true;
+	}
+
+	void RetakeSnapshotForWidgetInHierarchy(FWidgetTreeSnapshot& WidgetTreeSnapshot, UWidget& Widget, const FWidgetSnapshotSettings& SnapshotSettings)
+	{
+		if (FWidgetSnapshot* Snapshot = WidgetTreeSnapshot.WidgetSnapshots.Find(Widget.GetFName()))
+		{
+			*Snapshot = TakeWidgetSnapshot(Widget, SnapshotSettings.AllowedProperties);
+			return;
+		}
+		
+		FWidgetSnapshot* RootSnapshot = WidgetTreeSnapshot.WidgetSnapshots.Find(WidgetTreeSnapshot.RootWidget);
+		if (RootSnapshot && RootSnapshot->WidgetClass == Widget.GetClass())
+		{
+			*RootSnapshot = TakeWidgetSnapshot(Widget, SnapshotSettings.AllowedProperties);
+			return;
+		}
+		
+		UE_LOG(LogVCamCore, Warning, TEXT("Failed to retake snapshot of widget %s"), *Widget.GetPathName());
+	}
+
+	class FTakeWidgetSnapshotArchive : public FObjectWriter
+	{
+		const TSet<const FProperty*>& AllowedProperties;
+
+		static bool CheckIsDangerousProperty(const FProperty* InProperty, bool bLogWarnings)
+		{
+			// No delegates
+			if (InProperty->HasAnyPropertyFlags(CPF_BlueprintAssignable))
+			{
+				return true;
+			}
+			
+			// These properties have the potential to blow up our archive... if this is really needed in the future, adjust this code.
+			// By default, let's be careful because I have not particular case in mind where we'd want to save these properties.
+			if (InProperty->HasAnyPropertyFlags(CPF_InstancedReference | CPF_ContainsInstancedReference | CPF_PersistentInstance))
+			{
+				UE_CLOG(bLogWarnings, LogVCamCore, Warning, TEXT("Property %s was skipped because it could possibly bloat the snapshot!"), *InProperty->GetName());
+				return true;
+			}
+			
+			return false;
+		}
+	public:
+
+		FTakeWidgetSnapshotArchive(const TSet<const FProperty*>& AllowedProperties, TArray<uint8>& InBytes)
+			: FObjectWriter(InBytes)
+			, AllowedProperties(AllowedProperties)
+		{
+			SetIsPersistent(true);
+			SetIsSaving(true);
+		}
+
+		virtual bool ShouldSkipProperty(const FProperty* InProperty) const override
+		{
+			if (AllowedProperties.Contains(InProperty))
+			{
+				// Warn the caller that an explicitly specified property will be skipped...
+				return CheckIsDangerousProperty(InProperty, true);
+			}
+			// ... but do not warn if the property was not explicitly specified
+			if (CheckIsDangerousProperty(InProperty, false))
+			{
+				return true;
+			}
+			
+			const FArchiveSerializedPropertyChain* PropertyChain = GetSerializedPropertyChain();
+			if (!PropertyChain)
+			{
+				return true;
+			}
+
+			for (int32 i = 0; PropertyChain->GetNumProperties(); ++i)
+			{
+				// We expect allowed properties to be closer to the beginning than the end of the chain so start from beginning rather than from end
+				if (AllowedProperties.Contains(PropertyChain->GetPropertyFromRoot(i)))
+				{
+					// Child properties of an allowed property must not be skipped or they won't be saved.
+					return false;
+				}
+			}
+
+			return true;
+		}
+	};
+	
+	FWidgetSnapshot TakeWidgetSnapshot(UWidget& Widget, const TSet<const FProperty*>& AllowedProperties)
+	{
+		FWidgetSnapshot Result;
+		Result.WidgetClass = Widget.GetClass();
+		FTakeWidgetSnapshotArchive Archive(AllowedProperties, Result.SavedBinaryData);
+		FObjectAndNameAsStringProxyArchive Proxy(Archive, false);
+		// Not calling SetCustomVersions() because I do not expect these widgets to require any migration
+		Widget.Serialize(Proxy);
+		return Result;
+	}
+	
+	void ApplyWidgetSnapshot(const FWidgetSnapshot& Snapshot, UWidget& Widget)
+	{
+		FMemoryReader Reader(Snapshot.SavedBinaryData, true);
+		FObjectAndNameAsStringProxyArchive Proxy(Reader, false);
+		// Not calling SetCustomVersions() because I do not expect these widgets to require any migration
+		Widget.Serialize(Proxy);
+	}
+}
diff --git a/Engine/Plugins/Experimental/VirtualCamera/Source/VCamCore/Private/Util/WidgetSnapshotUtils.h b/Engine/Plugins/Experimental/VirtualCamera/Source/VCamCore/Private/Util/WidgetSnapshotUtils.h
new file mode 100644
--- /dev/null
+++ b/Engine/Plugins/Experimental/VirtualCamera/Source/VCamCore/Private/Util/WidgetSnapshotUtils.h
@@ -0,0 +1,32 @@
+// Copyright Epic Games, Inc. All Rights Reserved.
+
+#pragma once
+
+#include "CoreMinimal.h"
+#include "VCamCore.h"
+#include "UI/WidgetSnapshots.h"
+
+class UWidget;
+struct FWidgetTreeSnapshot;
+
+namespace UE::VCamCore::WidgetSnapshotUtils::Private
+{
+	struct FWidgetSnapshotSettings
+	{
+		TSet<TSubclassOf<UWidget>> AllowedWidgetClasses;
+		TSet<const FProperty*> AllowedProperties;
+
+		bool IsClassAllowed(TSubclassOf<UWidget> Widget) const;
+	};
+	
+	/** Saves the state of the widget */
+	FWidgetTreeSnapshot TakeTreeHierarchySnapshot(UUserWidget& Widget, const FWidgetSnapshotSettings& SnapshotSettings = VCamCore::Private::FVCamCoreModule::Get().GetSnapshotSettings());
+	/** Applies the data saved in Snapshot to Widget*/
+	bool ApplyTreeHierarchySnapshot(const FWidgetTreeSnapshot& Snapshot, UUserWidget& Widget);
+
+	/** Updates the snapshot only for this widget in the hierarchy if a matching snapshot exists (otherwise does not add it).*/
+	void RetakeSnapshotForWidgetInHierarchy(FWidgetTreeSnapshot& WidgetTreeSnapshot, UWidget& Widget, const FWidgetSnapshotSettings& SnapshotSettings = VCamCore::Private::FVCamCoreModule::Get().GetSnapshotSettings());
+
+	FWidgetSnapshot TakeWidgetSnapshot(UWidget& Widget, const TSet<const FProperty*>& AllowedProperties);
+	void ApplyWidgetSnapshot(const FWidgetSnapshot& Snapshot, UWidget& Widget);
+};
diff --git a/Engine/Plugins/Experimental/VirtualCamera/Source/VCamCore/Private/VCamComponent.cpp b/Engine/Plugins/Experimental/VirtualCamera/Source/VCamCore/Private/VCamComponent.cpp
--- a/Engine/Plugins/Experimental/VirtualCamera/Source/VCamCore/Private/VCamComponent.cpp
+++ b/Engine/Plugins/Experimental/VirtualCamera/Source/VCamCore/Private/VCamComponent.cpp
@@ -4,6 +4,7 @@
 
 #include "Modifier/VCamModifier.h"
 #include "Modifier/VCamModifierContext.h"
+#include "Output/VCamOutputProviderBase.h"
 #include "Util/LevelViewportUtils.h"
 #include "VCamCoreCustomVersion.h"
 
@@ -39,7 +40,6 @@
 #include "VPRolesSubsystem.h"
 #endif
 
-
 DEFINE_LOG_CATEGORY(LogVCamComponent);
 
 namespace VCamComponent
diff --git a/Engine/Plugins/Experimental/VirtualCamera/Source/VCamCore/Private/VCamCore.cpp b/Engine/Plugins/Experimental/VirtualCamera/Source/VCamCore/Private/VCamCore.cpp
--- a/Engine/Plugins/Experimental/VirtualCamera/Source/VCamCore/Private/VCamCore.cpp
+++ b/Engine/Plugins/Experimental/VirtualCamera/Source/VCamCore/Private/VCamCore.cpp
@@ -2,19 +2,35 @@
 
 #include "VCamCore.h"
 
+#include "UI/VCamWidget.h"
+#include "Util/WidgetSnapshotUtils.h"
+
 #define LOCTEXT_NAMESPACE "FVCamCoreModule"
 
-void FVCamCoreModule::StartupModule()
+namespace UE::VCamCore::Private
 {
-	// This code will execute after your module is loaded into memory; the exact timing is specified in the .uplugin file per-module
-}
+	void FVCamCoreModule::StartupModule()
+	{
+		
+	}
 
-void FVCamCoreModule::ShutdownModule()
-{
-	// This function may be called during shutdown to clean up your module.  For modules that support dynamic reloading,
-	// we call this function before unloading the module.
+	void FVCamCoreModule::ShutdownModule()
+	{
+		
+	}
+
+	WidgetSnapshotUtils::Private::FWidgetSnapshotSettings FVCamCoreModule::GetSnapshotSettings() const
+	{
+		// In the future this could be exposed via project settings or via registration functions on IVCamCoreModule
+		const TSet<TSubclassOf<UWidget>> AllowedWidgetClasses { UVCamWidget::StaticClass() };
+		const TSet<const FProperty*> AllowedProperties { UVCamWidget::StaticClass()->FindPropertyByName(GET_MEMBER_NAME_CHECKED(UVCamWidget, Connections)) };
+		return WidgetSnapshotUtils::Private::FWidgetSnapshotSettings{
+			AllowedWidgetClasses,
+			AllowedProperties
+		};
+	}
 }
 
 #undef LOCTEXT_NAMESPACE
 	
-IMPLEMENT_MODULE(FVCamCoreModule, VCamCore)
\ No newline at end of file
+IMPLEMENT_MODULE(UE::VCamCore::Private::FVCamCoreModule, VCamCore)
\ No newline at end of file
diff --git a/Engine/Plugins/Experimental/VirtualCamera/Source/VCamCore/Private/VCamCore.h b/Engine/Plugins/Experimental/VirtualCamera/Source/VCamCore/Private/VCamCore.h
new file mode 100644
--- /dev/null
+++ b/Engine/Plugins/Experimental/VirtualCamera/Source/VCamCore/Private/VCamCore.h
@@ -0,0 +1,33 @@
+// Copyright Epic Games, Inc. All Rights Reserved.
+
+#pragma once
+
+#include "CoreMinimal.h"
+#include "IVCamCoreModule.h"
+#include "Modules/ModuleManager.h"
+
+namespace UE::VCamCore::WidgetSnapshotUtils::Private
+{
+	struct FWidgetSnapshotSettings;
+}
+
+namespace UE::VCamCore::Private
+{
+	class FVCamCoreModule : public IVCamCoreModule
+	{
+	public:
+
+		static FVCamCoreModule& Get()
+		{
+			return FModuleManager::Get().GetModuleChecked<FVCamCoreModule>("VCamCore");
+		}
+
+		//~ Begin IModuleInterface Interface
+		virtual void StartupModule() override;
+		virtual void ShutdownModule() override;
+		//~ End IModuleInterface Interface
+
+		WidgetSnapshotUtils::Private::FWidgetSnapshotSettings GetSnapshotSettings() const;
+	};
+}
+
diff --git a/Engine/Plugins/Experimental/VirtualCamera/Source/VCamCore/Public/IVCamCoreModule.h b/Engine/Plugins/Experimental/VirtualCamera/Source/VCamCore/Public/IVCamCoreModule.h
new file mode 100644
--- /dev/null
+++ b/Engine/Plugins/Experimental/VirtualCamera/Source/VCamCore/Public/IVCamCoreModule.h
@@ -0,0 +1,19 @@
+// Copyright Epic Games, Inc. All Rights Reserved.
+
+#pragma once
+
+#include "CoreMinimal.h"
+#include "Modules/ModuleInterface.h"
+
+namespace UE::VCamCore
+{
+	class VCAMCORE_API IVCamCoreModule : public IModuleInterface
+	{
+	public:
+
+		static IVCamCoreModule& Get()
+		{
+			return FModuleManager::Get().GetModuleChecked<IVCamCoreModule>("VCamCore");
+		}
+	};
+}
\ No newline at end of file
diff --git a/Engine/Plugins/Experimental/VirtualCamera/Source/VCamCore/Public/Output/VCamOutputProviderBase.h b/Engine/Plugins/Experimental/VirtualCamera/Source/VCamCore/Public/Output/VCamOutputProviderBase.h
--- a/Engine/Plugins/Experimental/VirtualCamera/Source/VCamCore/Public/Output/VCamOutputProviderBase.h
+++ b/Engine/Plugins/Experimental/VirtualCamera/Source/VCamCore/Public/Output/VCamOutputProviderBase.h
@@ -6,14 +6,16 @@
 #include "CineCameraComponent.h"
 #include "EVCamTargetViewportID.h"
 #include "VPFullScreenUserWidget.h"
+#include "UI/WidgetSnapshots.h"
 #include "VCamOutputProviderBase.generated.h"
 
 DECLARE_LOG_CATEGORY_EXTERN(LogVCamOutputProvider, Log, All);
 
+class FSceneViewport;
+class SWindow;
 class UUserWidget;
+class UVCamWidget;
 class UVPFullScreenUserWidget;
-class SWindow;
-class FSceneViewport;
 
 #if WITH_EDITOR
 class FLevelEditorViewportClient;
@@ -169,12 +171,16 @@ private:
 	
 	UPROPERTY(Transient)
 	bool bWasActive = false;
-
-#if WITH_EDITORONLY_DATA
-	/** Used for saving widget remapping settings across enabling & disabling in the same session, see FOutputProviderLayoutCustomization. */
-	UPROPERTY(Transient)
-	TObjectPtr<UUserWidget> SavedConnectionRemappingData = nullptr;
-#endif
+	
+	/** FOutputProviderLayoutCustomization allows remapping connections and their bound widgets. This is used to persist those overrides since UUserWidgets cannot be saved. */
+	UPROPERTY()
+	FWidgetTreeSnapshot WidgetSnapshot;
 	
 	bool IsOuterComponentEnabled() const;
+
+#if WITH_EDITOR
+	void StartDetectAndSnapshotWhenConnectionsChange();
+	void StopDetectAndSnapshotWhenConnectionsChange();
+	void OnConnectionReinitialized(TWeakObjectPtr<UVCamWidget> Widget);
+#endif
 };
diff --git a/Engine/Plugins/Experimental/VirtualCamera/Source/VCamCore/Public/UI/VCamWidget.h b/Engine/Plugins/Experimental/VirtualCamera/Source/VCamCore/Public/UI/VCamWidget.h
--- a/Engine/Plugins/Experimental/VirtualCamera/Source/VCamCore/Public/UI/VCamWidget.h
+++ b/Engine/Plugins/Experimental/VirtualCamera/Source/VCamCore/Public/UI/VCamWidget.h
@@ -3,12 +3,11 @@
 #pragma once
 
 #include "CoreMinimal.h"
-#include "VCamComponent.h"
 #include "UI/VCamConnectionStructs.h"
 #include "Blueprint/UserWidget.h"
-
 #include "VCamWidget.generated.h"
 
+class UInputMappingContext;
 class UInputAction;
 class UVCamModifier;
 
diff --git a/Engine/Plugins/Experimental/VirtualCamera/Source/VCamCore/Public/UI/WidgetSnapshots.h b/Engine/Plugins/Experimental/VirtualCamera/Source/VCamCore/Public/UI/WidgetSnapshots.h
new file mode 100644
--- /dev/null
+++ b/Engine/Plugins/Experimental/VirtualCamera/Source/VCamCore/Public/UI/WidgetSnapshots.h
@@ -0,0 +1,43 @@
+// Copyright Epic Games, Inc. All Rights Reserved.
+
+#pragma once
+
+#include "CoreMinimal.h"
+#include "Templates/SubclassOf.h"
+#include "Blueprint/UserWidget.h"
+#include "WidgetSnapshots.generated.h"
+
+/** Saves a subset of properties of a widget. */
+USTRUCT()
+struct FWidgetSnapshot
+{
+	GENERATED_BODY()
+
+	UPROPERTY()
+	TSubclassOf<UUserWidget> WidgetClass;
+
+	/** Can be empty in case no properties were saved */
+	UPROPERTY()
+	TArray<uint8> SavedBinaryData;
+};
+
+/** Saves properties from widgets within a widget tree. */
+USTRUCT()
+struct FWidgetTreeSnapshot
+{
+	GENERATED_BODY()
+
+	/**
+	 * The name of the root widget if it was saved (may not have passed filters).
+	 * The root widget may be re-instanced (and thus renamed).
+	 */
+	UPROPERTY()
+	FName RootWidget = NAME_None;
+	
+	/** Stores every widget in the hierarchy which passed filters, even if there are no properties to save. */
+	UPROPERTY()
+	TMap<FName, FWidgetSnapshot> WidgetSnapshots;
+
+	void Reset() { WidgetSnapshots.Reset(); }
+	bool HasData() const { return WidgetSnapshots.Num() != 0; }
+};
\ No newline at end of file
diff --git a/Engine/Plugins/Experimental/VirtualCamera/Source/VCamCore/Public/VCamComponent.h b/Engine/Plugins/Experimental/VirtualCamera/Source/VCamCore/Public/VCamComponent.h
--- a/Engine/Plugins/Experimental/VirtualCamera/Source/VCamCore/Public/VCamComponent.h
+++ b/Engine/Plugins/Experimental/VirtualCamera/Source/VCamCore/Public/VCamComponent.h
@@ -6,7 +6,6 @@
 
 #include "EVCamTargetViewportID.h"
 #include "Modifier/ModifierStackEntry.h"
-#include "Output/VCamOutputProviderBase.h"
 
 #include "EnhancedActionKeyMapping.h"
 #include "EnhancedInputSubsystemInterface.h"
@@ -23,6 +22,7 @@
 
 #include "VCamComponent.generated.h"
 
+class UVCamOutputProviderBase;
 DECLARE_LOG_CATEGORY_EXTERN(LogVCamComponent, Log, All);
 
 class UCineCameraComponent;
diff --git a/Engine/Plugins/Experimental/VirtualCamera/Source/VCamCore/Public/VCamCore.h b/Engine/Plugins/Experimental/VirtualCamera/Source/VCamCore/Public/VCamCore.h
deleted file mode 100644
