Details panel keeps closing open sections after changing the connection

diff --git a/Engine/Plugins/Experimental/VirtualCamera/Source/VCamCoreEditor/Private/Customizations/OutputProvider/OutputProviderLayoutCustomization.cpp b/Engine/Plugins/Experimental/VirtualCamera/Source/VCamCoreEditor/Private/Customizations/OutputProvider/OutputProviderLayoutCustomization.cpp
--- a/Engine/Plugins/Experimental/VirtualCamera/Source/VCamCoreEditor/Private/Customizations/OutputProvider/OutputProviderLayoutCustomization.cpp
+++ b/Engine/Plugins/Experimental/VirtualCamera/Source/VCamCoreEditor/Private/Customizations/OutputProvider/OutputProviderLayoutCustomization.cpp
@@ -277,12 +277,9 @@ namespace UE::VCamCoreEditor::Private
 					else if (TSharedPtr<IConnectionRemapCustomization> Customization = FVCamCoreEditorModule::Get().CreateConnectionRemapCustomization(VCamWidget->GetClass()))
 					{
 						const TWeakObjectPtr<UVCamWidget> WeakVCamWidget(VCamWidget);
-						FDelegateHandle OnPostConnectionsReinitializedDelegateHandle = VCamWidget->OnPostConnectionsReinitializedDelegate.AddSP(
-							this, &FOutputProviderLayoutCustomization::OnPostConnectionsReinitialized
-							);
 						EditableWidgets.Emplace(
 							WeakVCamWidget,
-							FWidgetData{ OnPostConnectionsReinitializedDelegateHandle, Customization, MakeShared<FConnectionRemapUtilsImpl>(WeakDetailBuilder.Pin().ToSharedRef()) }
+							FWidgetData{ Customization, MakeShared<FConnectionRemapUtilsImpl>(WeakDetailBuilder.Pin().ToSharedRef()) }
 							);
 					}
 
@@ -369,11 +366,6 @@ namespace UE::VCamCoreEditor::Private
 	{
 		ForceRefreshDetailsIfSafe();
 	}
-	
-	void FOutputProviderLayoutCustomization::OnPostConnectionsReinitialized() const
-	{
-		ForceRefreshDetailsIfSafe();
-	}
 
 	void FOutputProviderLayoutCustomization::ForceRefreshDetailsIfSafe() const
 	{
@@ -385,13 +377,6 @@ namespace UE::VCamCoreEditor::Private
 
 	void FOutputProviderLayoutCustomization::ClearWidgetData(TMap<TWeakObjectPtr<UVCamWidget>, FWidgetData>& InEditableWidgets)
 	{
-		for (TPair<TWeakObjectPtr<UVCamWidget>, FWidgetData>& WidgetData : InEditableWidgets)
-		{
-			if (WidgetData.Key.IsValid())
-			{
-				WidgetData.Key->OnPostConnectionsReinitializedDelegate.Remove(WidgetData.Value.OnPostConnectionsReinitializedDelegateHandle);
-			}
-		}
 		InEditableWidgets.Empty();
 	}
 }
diff --git a/Engine/Plugins/Experimental/VirtualCamera/Source/VCamCoreEditor/Private/Customizations/OutputProvider/OutputProviderLayoutCustomization.h b/Engine/Plugins/Experimental/VirtualCamera/Source/VCamCoreEditor/Private/Customizations/OutputProvider/OutputProviderLayoutCustomization.h
--- a/Engine/Plugins/Experimental/VirtualCamera/Source/VCamCoreEditor/Private/Customizations/OutputProvider/OutputProviderLayoutCustomization.h
+++ b/Engine/Plugins/Experimental/VirtualCamera/Source/VCamCoreEditor/Private/Customizations/OutputProvider/OutputProviderLayoutCustomization.h
@@ -42,9 +42,6 @@ namespace UE::VCamCoreEditor::Private
 		
 		struct FWidgetData
 		{
-			/** Handle to UVCamWidget::OnPostConnectionsReinitialized */
-			FDelegateHandle OnPostConnectionsReinitializedDelegateHandle;
-
 			TSharedPtr<IConnectionRemapCustomization> Customization;
 			
 			/** Utils for IConnectionRemapCustomizations */
@@ -59,7 +56,6 @@ namespace UE::VCamCoreEditor::Private
 		TSharedRef<SHorizontalBox> CreateControlWidgets(const TWeakObjectPtr<UVCamWidget>& Widget) const;
 		
 		void OnActivationChanged(bool bNewIsActivated) const;
-		void OnPostConnectionsReinitialized() const;
 		void ForceRefreshDetailsIfSafe() const;
 		
 		static void ClearWidgetData(TMap<TWeakObjectPtr<UVCamWidget>, FWidgetData>& InEditableWidgets);
