VCam subwidget of VCam widgets are now also found and used. Persist remapping of connections.

diff --git a/Engine/Plugins/Experimental/VirtualCamera/Source/VCamCore/Private/Output/VCamOutputProviderBase.cpp b/Engine/Plugins/Experimental/VirtualCamera/Source/VCamCore/Private/Output/VCamOutputProviderBase.cpp
--- a/Engine/Plugins/Experimental/VirtualCamera/Source/VCamCore/Private/Output/VCamOutputProviderBase.cpp
+++ b/Engine/Plugins/Experimental/VirtualCamera/Source/VCamCore/Private/Output/VCamOutputProviderBase.cpp
@@ -16,6 +16,7 @@
 #include "UObject/UObjectBaseUtility.h"
 #include "Util/ConnectionUtils.h"
 #include "Util/WidgetSnapshotUtils.h"
+#include "Util/WidgetTreeUtils.h"
 
 #if WITH_EDITOR
 #include "Editor.h"
@@ -316,26 +317,22 @@ void UVCamOutputProviderBase::NotifyWidgetOfComponentChange() const
 			if (UVCamComponent* OwningComponent = GetTypedOuter<UVCamComponent>())
 			{
 				UVCamComponent* VCamComponent = bIsActive ? OwningComponent : nullptr;
-
-				if (DisplayedWidget->Implements<UVCamModifierInterface>())
-				{
-					IVCamModifierInterface::Execute_OnVCamComponentChanged(DisplayedWidget, VCamComponent);
-				}
-
+				
 				// Find all VCam Widgets inside the displayed widget and Initialize them with the owning VCam Component
-				if (IsValid(DisplayedWidget->WidgetTree))
+				UE::VCamCore::ForEachWidgetToConsiderForVCam(*DisplayedWidget, [VCamComponent](UWidget* Widget)
 				{
-					DisplayedWidget->WidgetTree->ForEachWidget([VCamComponent](UWidget* Widget)
+					if (UVCamWidget* VCamWidget = Cast<UVCamWidget>(Widget))
 					{
-						if (UVCamWidget* VCamWidget = Cast<UVCamWidget>(Widget))
-						{
-							VCamWidget->InitializeConnections(VCamComponent);
-						}
-					});
-				}
+						VCamWidget->InitializeConnections(VCamComponent);
+					}
+					
+					if (Widget->Implements<UVCamModifierInterface>())
+					{
+						IVCamModifierInterface::Execute_OnVCamComponentChanged(Widget, VCamComponent);
+					}
+				});
 			}
 		}
-
 	}
 }
 
@@ -551,7 +548,7 @@ void UVCamOutputProviderBase::StartDetectAndSnapshotWhenConnectionsChange()
 	UUserWidget* Widget = UMGWidget->GetWidget();
 	check(Widget);
 	
-	Widget->WidgetTree->ForEachWidget([this](UWidget* Widget)
+	UE::VCamCore::ForEachWidgetToConsiderForVCam(*Widget, [this](UWidget* Widget)
 	{
 		if (UVCamWidget* VCamWidget = Cast<UVCamWidget>(Widget))
 		{
@@ -567,11 +564,10 @@ void UVCamOutputProviderBase::StopDetectAndSnapshotWhenConnectionsChange()
 	UUserWidget* Widget = UMGWidget->GetWidget();
 	check(Widget);
 	
-	Widget->WidgetTree->ForEachWidget([this](UWidget* Widget)
+	UE::VCamCore::ForEachWidgetToConsiderForVCam(*Widget, [this](UWidget* Widget)
 	{
 		if (UVCamWidget* VCamWidget = Cast<UVCamWidget>(Widget))
 		{
-			const TWeakObjectPtr<UVCamWidget> WeakWidget = VCamWidget;
 			VCamWidget->OnPostConnectionsReinitializedDelegate.RemoveAll(this);
 		}
 	});
diff --git a/Engine/Plugins/Experimental/VirtualCamera/Source/VCamCore/Private/Util/WidgetSnapshotUtils.cpp b/Engine/Plugins/Experimental/VirtualCamera/Source/VCamCore/Private/Util/WidgetSnapshotUtils.cpp
--- a/Engine/Plugins/Experimental/VirtualCamera/Source/VCamCore/Private/Util/WidgetSnapshotUtils.cpp
+++ b/Engine/Plugins/Experimental/VirtualCamera/Source/VCamCore/Private/Util/WidgetSnapshotUtils.cpp
@@ -10,6 +10,7 @@
 #include "Serialization/ObjectReader.h"
 #include "Serialization/ObjectWriter.h"
 #include "UI/WidgetSnapshots.h"
+#include "Util/WidgetTreeUtils.h"
 
 namespace UE::VCamCore::WidgetSnapshotUtils::Private
 {
@@ -33,7 +34,7 @@ namespace UE::VCamCore::WidgetSnapshotUtils::Private
 		}
 
 		FWidgetTreeSnapshot Result;
-		Widget.WidgetTree->ForEachWidget([&Result, &SnapshotSettings](UWidget* Widget)
+		UE::VCamCore::ForEachWidgetToConsiderForVCam(Widget, [&Result, &SnapshotSettings](UWidget* Widget)
 		{
 			if (Widget && SnapshotSettings.IsClassAllowed(Widget->GetClass()))
 			{
@@ -58,7 +59,7 @@ namespace UE::VCamCore::WidgetSnapshotUtils::Private
 			return false;
 		}
 
-		Widget.WidgetTree->ForEachWidget([&Snapshot](UWidget* Widget)
+		UE::VCamCore::ForEachWidgetToConsiderForVCam(Widget,[&Snapshot](UWidget* Widget)
 		{
 			if (const FWidgetSnapshot* WidgetSnapshot = Snapshot.WidgetSnapshots.Find(Widget->GetFName()))
 			{
diff --git a/Engine/Plugins/Experimental/VirtualCamera/Source/VCamCore/Private/Util/WidgetTreeUtils.cpp b/Engine/Plugins/Experimental/VirtualCamera/Source/VCamCore/Private/Util/WidgetTreeUtils.cpp
new file mode 100644
--- /dev/null
+++ b/Engine/Plugins/Experimental/VirtualCamera/Source/VCamCore/Private/Util/WidgetTreeUtils.cpp
@@ -0,0 +1,34 @@
+// Copyright Epic Games, Inc. All Rights Reserved.
+
+#include "Util/WidgetTreeUtils.h"
+
+#include "Blueprint/WidgetTree.h"
+#include "UI/VCamWidget.h"
+
+namespace UE::VCamCore
+{
+	void ForEachWidgetToConsiderForVCam(UUserWidget& Widget, TFunctionRef<void(UWidget* Widget)> Callback)
+	{
+		if (!Widget.WidgetTree)
+		{
+			return;
+		}
+
+		Callback(&Widget);
+					
+		TQueue<UWidgetTree*> SearchQueue;
+		SearchQueue.Enqueue(Widget.WidgetTree);
+		UWidgetTree* CurrentTree;
+		while (SearchQueue.Dequeue(CurrentTree))
+		{
+			CurrentTree->ForEachWidget([&SearchQueue, Callback](UWidget* Widget)
+			{
+				Callback(Widget);
+				if (UUserWidget* UserWidget = Cast<UUserWidget>(Widget); UserWidget && UserWidget->WidgetTree)
+				{
+					SearchQueue.Enqueue(UserWidget->WidgetTree);
+				}
+			});
+		}
+	}
+}
diff --git a/Engine/Plugins/Experimental/VirtualCamera/Source/VCamCore/Public/Util/WidgetTreeUtils.h b/Engine/Plugins/Experimental/VirtualCamera/Source/VCamCore/Public/Util/WidgetTreeUtils.h
new file mode 100644
--- /dev/null
+++ b/Engine/Plugins/Experimental/VirtualCamera/Source/VCamCore/Public/Util/WidgetTreeUtils.h
@@ -0,0 +1,14 @@
+// Copyright Epic Games, Inc. All Rights Reserved.
+
+#pragma once
+
+#include "CoreMinimal.h"
+
+class UUserWidget;
+class UWidget;
+
+namespace UE::VCamCore
+{
+	/** Calls the callback for each widget, including the children of UUserWidgets. */
+	VCAMCORE_API void ForEachWidgetToConsiderForVCam(UUserWidget& Widget, TFunctionRef<void(UWidget*)> Callback);
+}
\ No newline at end of file
diff --git a/Engine/Plugins/Experimental/VirtualCamera/Source/VCamCoreEditor/Private/Customizations/OutputProvider/OutputProviderLayoutCustomization.cpp b/Engine/Plugins/Experimental/VirtualCamera/Source/VCamCoreEditor/Private/Customizations/OutputProvider/OutputProviderLayoutCustomization.cpp
--- a/Engine/Plugins/Experimental/VirtualCamera/Source/VCamCoreEditor/Private/Customizations/OutputProvider/OutputProviderLayoutCustomization.cpp
+++ b/Engine/Plugins/Experimental/VirtualCamera/Source/VCamCoreEditor/Private/Customizations/OutputProvider/OutputProviderLayoutCustomization.cpp
@@ -14,8 +14,10 @@
 #include "Editor.h"
 #include "Framework/MultiBox/MultiBoxBuilder.h"
 #include "IDetailGroup.h"
+#include "LogVCamEditor.h"
 #include "PropertyCustomizationHelpers.h"
 #include "SSimpleComboButton.h"
+#include "Util/WidgetTreeUtils.h"
 #include "Widgets/Images/SImage.h"
 #include "Widgets/Input/SButton.h"
 #include "Widgets/Layout/SBox.h"
@@ -249,7 +251,7 @@ namespace UE::VCamCoreEditor::Private
 	void FOutputProviderLayoutCustomization::RebuildWidgetData()
 	{
 		const UVPFullScreenUserWidget* FullScreenUserWidget = CustomizedOutputProvider->GetUMGWidget();
-		const UUserWidget* RootWidget = FullScreenUserWidget ? FullScreenUserWidget->GetWidget() : nullptr;
+		UUserWidget* RootWidget = FullScreenUserWidget ? FullScreenUserWidget->GetWidget() : nullptr;
 		UWidgetTree* WidgetTree = RootWidget ? RootWidget->WidgetTree : nullptr;
 		if (!WidgetTree)
 		{
@@ -258,38 +260,25 @@ namespace UE::VCamCoreEditor::Private
 
 		// Move so EditableWidgets gets reset and does not retain old references
 		TMap<TWeakObjectPtr<UVCamWidget>, FWidgetData> OldEditableWidgets = MoveTemp(EditableWidgets);
-		
-		TQueue<UWidgetTree*> SearchQueue;
-		SearchQueue.Enqueue(WidgetTree);
-		UWidgetTree* CurrentTree;
-		while (SearchQueue.Dequeue(CurrentTree))
+		VCamCore::ForEachWidgetToConsiderForVCam(*RootWidget, [this, &OldEditableWidgets](UWidget* Widget)
 		{
-			// We recurse ourselves because the other functions skip UUserWidgets
-			CurrentTree->ForEachWidget([this, &OldEditableWidgets, &SearchQueue](UWidget* Widget)
+			if (UVCamWidget* VCamWidget = Cast<UVCamWidget>(Widget))
 			{
-				if (UVCamWidget* VCamWidget = Cast<UVCamWidget>(Widget))
+				FWidgetData ExistingWidgetData;
+				if (OldEditableWidgets.RemoveAndCopyValue(VCamWidget, ExistingWidgetData))
 				{
-					FWidgetData ExistingWidgetData;
-					if (OldEditableWidgets.RemoveAndCopyValue(VCamWidget, ExistingWidgetData))
-					{
-						EditableWidgets.Emplace(VCamWidget, MoveTemp(ExistingWidgetData));
-					}
-					else if (TSharedPtr<IConnectionRemapCustomization> Customization = FVCamCoreEditorModule::Get().CreateConnectionRemapCustomization(VCamWidget->GetClass()))
-					{
-						const TWeakObjectPtr<UVCamWidget> WeakVCamWidget(VCamWidget);
-						EditableWidgets.Emplace(
-							WeakVCamWidget,
-							FWidgetData{ Customization, MakeShared<FConnectionRemapUtilsImpl>(WeakDetailBuilder.Pin().ToSharedRef()) }
-							);
-					}
-
-					if (VCamWidget->WidgetTree)
-					{
-						SearchQueue.Enqueue(VCamWidget->WidgetTree);
-					}
+					EditableWidgets.Emplace(VCamWidget, MoveTemp(ExistingWidgetData));
 				}
-			});
-		}
+				else if (TSharedPtr<IConnectionRemapCustomization> Customization = FVCamCoreEditorModule::Get().CreateConnectionRemapCustomization(VCamWidget->GetClass()))
+				{
+					const TWeakObjectPtr<UVCamWidget> WeakVCamWidget(VCamWidget);
+					EditableWidgets.Emplace(
+						WeakVCamWidget,
+						FWidgetData{ Customization, MakeShared<FConnectionRemapUtilsImpl>(WeakDetailBuilder.Pin().ToSharedRef()) }
+						);
+				}
+			}
+		});
 		
 		// The hierarchy may have changed so anything that is left is not part of the hierarchy and can be unsubscribed from
 		ClearWidgetData(OldEditableWidgets);
