State switcher remapping customization

diff --git a/Engine/Plugins/Experimental/VirtualCamera/Source/VCamCore/Private/UI/Switcher/WidgetConnectionConfig.cpp b/Engine/Plugins/Experimental/VirtualCamera/Source/VCamCore/Private/UI/Switcher/WidgetConnectionConfig.cpp
--- a/Engine/Plugins/Experimental/VirtualCamera/Source/VCamCore/Private/UI/Switcher/WidgetConnectionConfig.cpp
+++ b/Engine/Plugins/Experimental/VirtualCamera/Source/VCamCore/Private/UI/Switcher/WidgetConnectionConfig.cpp
@@ -2,27 +2,11 @@
 
 #include "UI/Switcher/WidgetConnectionConfig.h"
 
-#include "Blueprint/WidgetTree.h"
 #include "LogVCamCore.h"
 #include "UI/Switcher/VCamStateSwitcherWidget.h"
+#include "Util/WidgetTreeUtils.h"
 
-#if WITH_EDITOR
-#include "BaseWidgetBlueprint.h"
-#endif
-
-namespace UE::VCamCore::Private
-{
-#if WITH_EDITOR
-	static UWidgetTree* GetWidgetTreeThroughBlueprintAsset(UUserWidget* ClassDefaultWidget)
-	{
-		UObject* Blueprint = ClassDefaultWidget->GetClass()->ClassGeneratedBy;
-		UBaseWidgetBlueprint* WidgetBlueprint = Cast<UBaseWidgetBlueprint>(Blueprint);
-		return WidgetBlueprint
-			? WidgetBlueprint->WidgetTree
-			: nullptr;
-	}
-#endif
-}
+#include "Blueprint/WidgetTree.h"
 
 UVCamWidget* FWidgetConnectionConfig::ResolveWidget(UVCamStateSwitcherWidget* OwnerWidget) const
 {
@@ -35,7 +19,7 @@ UVCamWidget* FWidgetConnectionConfig::ResolveWidget(UVCamStateSwitcherWidget* Ow
 #if WITH_EDITOR
 	if (!WidgetTree && OwnerWidget->HasAnyFlags(RF_ClassDefaultObject))
 	{
-		WidgetTree = UE::VCamCore::Private::GetWidgetTreeThroughBlueprintAsset(OwnerWidget);
+		WidgetTree = UE::VCamCore::GetWidgetTreeThroughBlueprintAsset(*OwnerWidget);
 	}
 #endif
 
diff --git a/Engine/Plugins/Experimental/VirtualCamera/Source/VCamCore/Private/UI/VCamConnectionStructs.cpp b/Engine/Plugins/Experimental/VirtualCamera/Source/VCamCore/Private/UI/VCamConnectionStructs.cpp
--- a/Engine/Plugins/Experimental/VirtualCamera/Source/VCamCore/Private/UI/VCamConnectionStructs.cpp
+++ b/Engine/Plugins/Experimental/VirtualCamera/Source/VCamCore/Private/UI/VCamConnectionStructs.cpp
@@ -33,7 +33,7 @@ bool FVCamConnection::IsConnectionValid(UVCamModifier& Modifier, FName Connectio
 {
 	const bool bImplementsAllRequiredInterfaces = Algo::AllOf(RequiredInterfaces, [&Modifier](const TSubclassOf<UInterface>& Interface)
 	{
-		return Modifier.GetClass()->ImplementsInterface(Interface);
+		return !Interface->GetClass() || Modifier.GetClass()->ImplementsInterface(Interface);
 	});
 	if (!bImplementsAllRequiredInterfaces)
 	{
diff --git a/Engine/Plugins/Experimental/VirtualCamera/Source/VCamCore/Private/Util/WidgetTreeUtils.cpp b/Engine/Plugins/Experimental/VirtualCamera/Source/VCamCore/Private/Util/WidgetTreeUtils.cpp
--- a/Engine/Plugins/Experimental/VirtualCamera/Source/VCamCore/Private/Util/WidgetTreeUtils.cpp
+++ b/Engine/Plugins/Experimental/VirtualCamera/Source/VCamCore/Private/Util/WidgetTreeUtils.cpp
@@ -2,9 +2,14 @@
 
 #include "Util/WidgetTreeUtils.h"
 
-#include "Blueprint/WidgetTree.h"
 #include "UI/VCamWidget.h"
 
+#include "Blueprint/WidgetTree.h"
+
+#if WITH_EDITOR
+#include "BaseWidgetBlueprint.h"
+#endif
+
 namespace UE::VCamCore
 {
 	void ForEachWidgetToConsiderForVCam(UUserWidget& Widget, TFunctionRef<void(UWidget* Widget)> Callback)
@@ -31,4 +36,15 @@ namespace UE::VCamCore
 			});
 		}
 	}
+
+#if WITH_EDITOR
+	UWidgetTree* GetWidgetTreeThroughBlueprintAsset(UUserWidget& Widget)
+	{
+		UObject* Blueprint = Widget.GetClass()->ClassGeneratedBy;
+		UBaseWidgetBlueprint* WidgetBlueprint = Cast<UBaseWidgetBlueprint>(Blueprint);
+		return WidgetBlueprint
+			? WidgetBlueprint->WidgetTree
+			: nullptr;
+	}
+#endif
 }
diff --git a/Engine/Plugins/Experimental/VirtualCamera/Source/VCamCore/Public/UI/Switcher/VCamStateSwitcherWidget.h b/Engine/Plugins/Experimental/VirtualCamera/Source/VCamCore/Public/UI/Switcher/VCamStateSwitcherWidget.h
--- a/Engine/Plugins/Experimental/VirtualCamera/Source/VCamCore/Public/UI/Switcher/VCamStateSwitcherWidget.h
+++ b/Engine/Plugins/Experimental/VirtualCamera/Source/VCamCore/Public/UI/Switcher/VCamStateSwitcherWidget.h
@@ -49,6 +49,8 @@ public:
 	virtual void PostEditChangeProperty(FPropertyChangedEvent& PropertyChangedEvent) override;
 #endif
 
+	FORCEINLINE static FName GetCurrentStatePropertyName() { return GET_MEMBER_NAME_CHECKED(UVCamStateSwitcherWidget, CurrentState); }
+
 protected:
 
 	//~ Begin UUserWidget Interface
diff --git a/Engine/Plugins/Experimental/VirtualCamera/Source/VCamCore/Public/Util/WidgetTreeUtils.h b/Engine/Plugins/Experimental/VirtualCamera/Source/VCamCore/Public/Util/WidgetTreeUtils.h
--- a/Engine/Plugins/Experimental/VirtualCamera/Source/VCamCore/Public/Util/WidgetTreeUtils.h
+++ b/Engine/Plugins/Experimental/VirtualCamera/Source/VCamCore/Public/Util/WidgetTreeUtils.h
@@ -4,6 +4,7 @@
 
 #include "CoreMinimal.h"
 
+class UWidgetTree;
 class UUserWidget;
 class UWidget;
 
@@ -11,4 +12,8 @@ namespace UE::VCamCore
 {
 	/** Calls the callback for each widget, including the children of UUserWidgets. */
 	VCAMCORE_API void ForEachWidgetToConsiderForVCam(UUserWidget& Widget, TFunctionRef<void(UWidget*)> Callback);
+
+#if WITH_EDITOR
+	VCAMCORE_API UWidgetTree* GetWidgetTreeThroughBlueprintAsset(UUserWidget& Widget);
+#endif
 }
\ No newline at end of file
diff --git a/Engine/Plugins/Experimental/VirtualCamera/Source/VCamCoreEditor/Private/Customizations/ConnectionTargetSettingsTypeCustomization.cpp b/Engine/Plugins/Experimental/VirtualCamera/Source/VCamCoreEditor/Private/Customizations/ConnectionTargetSettingsTypeCustomization.cpp
--- a/Engine/Plugins/Experimental/VirtualCamera/Source/VCamCoreEditor/Private/Customizations/ConnectionTargetSettingsTypeCustomization.cpp
+++ b/Engine/Plugins/Experimental/VirtualCamera/Source/VCamCoreEditor/Private/Customizations/ConnectionTargetSettingsTypeCustomization.cpp
@@ -73,8 +73,8 @@ namespace UE::VCamCoreEditor::Private
 		
 		const TSharedPtr<IPropertyHandle> ParentHandle = PropertyHandle->GetParentHandle();
 		const bool bIsValidHandle = ParentHandle
-			|| !CastField<FStructProperty>(PropertyHandle->GetProperty())
-			|| CastField<FStructProperty>(PropertyHandle->GetProperty())->Struct == FVCamConnection::StaticStruct();
+			&& CastField<FStructProperty>(PropertyHandle->GetProperty())
+			&& CastField<FStructProperty>(PropertyHandle->GetProperty())->Struct == FVCamConnection::StaticStruct();
 		const TSharedPtr<IPropertyHandle> OptionalVCamConnectionParentStructHandle = bIsValidHandle
 			? ParentHandle
 			: nullptr;
diff --git a/Engine/Plugins/Experimental/VirtualCamera/Source/VCamCoreEditor/Private/Customizations/OutputProvider/ConnectionRemapCustomization_StateSwitcher.cpp b/Engine/Plugins/Experimental/VirtualCamera/Source/VCamCoreEditor/Private/Customizations/OutputProvider/ConnectionRemapCustomization_StateSwitcher.cpp
new file mode 100644
--- /dev/null
+++ b/Engine/Plugins/Experimental/VirtualCamera/Source/VCamCoreEditor/Private/Customizations/OutputProvider/ConnectionRemapCustomization_StateSwitcher.cpp
@@ -0,0 +1,58 @@
+// Copyright Epic Games, Inc. All Rights Reserved.
+
+#include "ConnectionRemapCustomization_StateSwitcher.h"
+
+#include "DetailLayoutBuilder.h"
+#include "Customization/IConnectionRemapCustomization.h"
+#include "UI/VCamWidget.h"
+#include "UI/Switcher/VCamStateSwitcherWidget.h"
+
+#include "Algo/AnyOf.h"
+#include "DetailWidgetRow.h"
+#include "IDetailGroup.h"
+
+#define LOCTEXT_NAMESPACE "FConnectionRemapCustomization_StateSwitcher"
+
+namespace UE::VCamCoreEditor::Private
+{
+	TSharedRef<IConnectionRemapCustomization> FConnectionRemapCustomization_StateSwitcher::Make()
+	{
+		return MakeShared<FConnectionRemapCustomization_StateSwitcher>();
+	}
+
+	bool FConnectionRemapCustomization_StateSwitcher::CanGenerateGroup(const FShouldGenerateArgs& Args) const
+	{
+		UVCamStateSwitcherWidget* StateSwitcherWidget = Cast<UVCamStateSwitcherWidget>(Args.CustomizedWidget);
+		const bool bHasStateSwitcherProperties = ensure(StateSwitcherWidget)
+			|| Algo::AnyOf(StateSwitcherWidget->GetStates(), [StateSwitcherWidget](FName State)
+			{
+				FVCamWidgetConnectionState StateInfo;
+				return StateSwitcherWidget->GetStateInfo(State, StateInfo)
+					&& StateInfo.WidgetConfigs.Num()
+					&& Algo::AnyOf(StateInfo.WidgetConfigs, [](const FWidgetConnectionConfig& Config){ return Config.ConnectionTargets.Num() > 0; });
+			});
+
+		return bHasStateSwitcherProperties || Super::CanGenerateGroup(Args);
+	}
+	
+	void FConnectionRemapCustomization_StateSwitcher::Customize(const FConnectionRemapCustomizationArgs& Args)
+	{
+		AddCurrentStateProperty(Args);
+		Super::Customize(Args);
+	}
+
+	void FConnectionRemapCustomization_StateSwitcher::AddCurrentStateProperty(const FConnectionRemapCustomizationArgs& Args)
+	{
+		UVCamStateSwitcherWidget* StateSwitcherWidget = Cast<UVCamStateSwitcherWidget>(Args.CustomizedWidget);
+		const TSharedPtr<IPropertyHandle> CurrentStatePropertyHandle = Args.Builder.AddObjectPropertyData({ StateSwitcherWidget }, UVCamStateSwitcherWidget::GetCurrentStatePropertyName());
+		if (!ensure(CurrentStatePropertyHandle))
+		{
+			return;
+		}
+
+		Args.WidgetGroup.AddPropertyRow(CurrentStatePropertyHandle.ToSharedRef());
+	}
+
+}
+
+#undef LOCTEXT_NAMESPACE
\ No newline at end of file
diff --git a/Engine/Plugins/Experimental/VirtualCamera/Source/VCamCoreEditor/Private/Customizations/OutputProvider/ConnectionRemapCustomization_StateSwitcher.h b/Engine/Plugins/Experimental/VirtualCamera/Source/VCamCoreEditor/Private/Customizations/OutputProvider/ConnectionRemapCustomization_StateSwitcher.h
new file mode 100644
--- /dev/null
+++ b/Engine/Plugins/Experimental/VirtualCamera/Source/VCamCoreEditor/Private/Customizations/OutputProvider/ConnectionRemapCustomization_StateSwitcher.h
@@ -0,0 +1,27 @@
+// Copyright Epic Games, Inc. All Rights Reserved.
+
+#pragma once
+
+#include "CoreMinimal.h"
+#include "Customizations/OutputProvider/ConnectionRemapCustomization_VCamWidget.h"
+
+namespace UE::VCamCoreEditor::Private
+{
+	class FConnectionRemapCustomization_StateSwitcher : public FConnectionRemapCustomization_VCamWidget
+	{
+		using Super = FConnectionRemapCustomization_VCamWidget;
+	public:
+
+		static TSharedRef<IConnectionRemapCustomization> Make();
+		
+		//~ Begin IConnectionRemappingCustomization Interface
+		virtual bool CanGenerateGroup(const FShouldGenerateArgs& Args) const override;
+		virtual void Customize(const FConnectionRemapCustomizationArgs& Args) override;
+		//~ End IConnectionRemappingCustomization Interface
+
+	private:
+		
+		void AddCurrentStateProperty(const FConnectionRemapCustomizationArgs& Args);
+	};
+}
+
diff --git a/Engine/Plugins/Experimental/VirtualCamera/Source/VCamCoreEditor/Private/Customizations/OutputProvider/ConnectionRemapUtilsImpl.cpp b/Engine/Plugins/Experimental/VirtualCamera/Source/VCamCoreEditor/Private/Customizations/OutputProvider/ConnectionRemapUtilsImpl.cpp
--- a/Engine/Plugins/Experimental/VirtualCamera/Source/VCamCoreEditor/Private/Customizations/OutputProvider/ConnectionRemapUtilsImpl.cpp
+++ b/Engine/Plugins/Experimental/VirtualCamera/Source/VCamCoreEditor/Private/Customizations/OutputProvider/ConnectionRemapUtilsImpl.cpp
@@ -32,8 +32,6 @@ namespace UE::VCamCoreEditor::Private
 			// Overwrite bManuallyConfigureConnection because otherwise its EditCondition will hide the property!
 			DummyContainer.Connection.bManuallyConfigureConnection = true;
 			TStructOnScope<FConnectionContainerDummy> StructOnScope(DummyContainer);
-
-			// Object slicing occurs intentionally
 			StructData = MakeShared<TStructOnScope<FConnectionContainerDummy>>(MoveTemp(StructOnScope));
 			AddedConnections.Add(Args.ConnectionName, StructData);
 		}
diff --git a/Engine/Plugins/Experimental/VirtualCamera/Source/VCamCoreEditor/Private/Customizations/OutputProvider/OutputProviderLayoutCustomization.cpp b/Engine/Plugins/Experimental/VirtualCamera/Source/VCamCoreEditor/Private/Customizations/OutputProvider/OutputProviderLayoutCustomization.cpp
--- a/Engine/Plugins/Experimental/VirtualCamera/Source/VCamCoreEditor/Private/Customizations/OutputProvider/OutputProviderLayoutCustomization.cpp
+++ b/Engine/Plugins/Experimental/VirtualCamera/Source/VCamCoreEditor/Private/Customizations/OutputProvider/OutputProviderLayoutCustomization.cpp
@@ -324,7 +324,7 @@ namespace UE::VCamCoreEditor::Private
 				[
 					CreateControlWidgets(Widget)
 				];
-			WidgetData->Customization->Customize({ WidgetGroup, WidgetData->RemapUtils.ToSharedRef(), Widget, Private::DisplaySettings });
+			WidgetData->Customization->Customize({ DetailBuilder, WidgetGroup, WidgetData->RemapUtils.ToSharedRef(), Widget, Private::DisplaySettings });
 		}
 	}
 	
diff --git a/Engine/Plugins/Experimental/VirtualCamera/Source/VCamCoreEditor/Private/Customizations/StateSwitcher/WidgetConnectionConfigTypeCustomization.cpp b/Engine/Plugins/Experimental/VirtualCamera/Source/VCamCoreEditor/Private/Customizations/StateSwitcher/WidgetConnectionConfigTypeCustomization.cpp
--- a/Engine/Plugins/Experimental/VirtualCamera/Source/VCamCoreEditor/Private/Customizations/StateSwitcher/WidgetConnectionConfigTypeCustomization.cpp
+++ b/Engine/Plugins/Experimental/VirtualCamera/Source/VCamCoreEditor/Private/Customizations/StateSwitcher/WidgetConnectionConfigTypeCustomization.cpp
@@ -15,6 +15,7 @@
 #include "IPropertyUtilities.h"
 #include "Blueprint/WidgetTree.h"
 #include "UObject/PropertyIterator.h"
+#include "Util/WidgetTreeUtils.h"
 
 #define LOCTEXT_NAMESPACE "FWidgetConnectionConfigCustomization"
 
@@ -153,6 +154,7 @@ namespace UE::VCamCoreEditor::Private
 				PropertyItemList.Add(Widget->GetName());
 			}
 		});
+		
 		return PropertyItemList;
 	}
 
@@ -163,11 +165,7 @@ namespace UE::VCamCoreEditor::Private
 			return nullptr;
 		}
 		
-		UObject* Blueprint = ClassDefaultWidget->GetClass()->ClassGeneratedBy;
-		UBaseWidgetBlueprint* WidgetBlueprint = Cast<UBaseWidgetBlueprint>(Blueprint);
-		return WidgetBlueprint
-			? WidgetBlueprint->WidgetTree
-			: nullptr;
+		return VCamCore::GetWidgetTreeThroughBlueprintAsset(*ClassDefaultWidget);
 	}
 
 	void FWidgetConnectionConfigTypeCustomization::CustomizeConnectionTargetsReferenceProperty(
diff --git a/Engine/Plugins/Experimental/VirtualCamera/Source/VCamCoreEditor/Private/VCamCoreEditorModule.cpp b/Engine/Plugins/Experimental/VirtualCamera/Source/VCamCoreEditor/Private/VCamCoreEditorModule.cpp
--- a/Engine/Plugins/Experimental/VirtualCamera/Source/VCamCoreEditor/Private/VCamCoreEditorModule.cpp
+++ b/Engine/Plugins/Experimental/VirtualCamera/Source/VCamCoreEditor/Private/VCamCoreEditorModule.cpp
@@ -20,6 +20,7 @@
 #include "Modules/ModuleManager.h"
 #include "PropertyEditorModule.h"
 #include "Customizations/VCamViewportLockerTypeCustomization.h"
+#include "Customizations/OutputProvider/ConnectionRemapCustomization_StateSwitcher.h"
 #include "Customizations/OutputProvider/ConnectionRemapCustomization_VCamWidget.h"
 #include "Customizations/OutputProvider/OutputProviderTypeCustomization.h"
 #include "Output/VCamOutputViewport.h"
@@ -146,6 +147,10 @@ namespace UE::VCamCoreEditor::Private
 			UVCamWidget::StaticClass(),
 			FGetConnectionRemappingCustomization::CreateStatic(&FConnectionRemapCustomization_VCamWidget::Make)
 			);
+		RegisterConnectionRemapCustomization(
+			UVCamStateSwitcherWidget::StaticClass(),
+			FGetConnectionRemappingCustomization::CreateStatic(&FConnectionRemapCustomization_StateSwitcher::Make)
+			);
 	}
 
 	void FVCamCoreEditorModule::UnregisterCustomizations()
diff --git a/Engine/Plugins/Experimental/VirtualCamera/Source/VCamCoreEditor/Public/Customization/IConnectionRemapCustomization.h b/Engine/Plugins/Experimental/VirtualCamera/Source/VCamCoreEditor/Public/Customization/IConnectionRemapCustomization.h
--- a/Engine/Plugins/Experimental/VirtualCamera/Source/VCamCoreEditor/Public/Customization/IConnectionRemapCustomization.h
+++ b/Engine/Plugins/Experimental/VirtualCamera/Source/VCamCoreEditor/Public/Customization/IConnectionRemapCustomization.h
@@ -6,6 +6,7 @@
 #include "UObject/WeakObjectPtr.h"
 #include "UObject/WeakObjectPtrTemplates.h"
 
+class IDetailLayoutBuilder;
 class IDetailGroup;
 class UVCamWidget;
 
@@ -26,8 +27,11 @@ namespace UE::VCamCoreEditor
 
 	struct FConnectionRemapCustomizationArgs
 	{
+		/** Use to register custom layouts or add external object / struct data. */
+		IDetailLayoutBuilder& Builder;
 		/** Builder for the group under which to add the widget's data. */
 		IDetailGroup& WidgetGroup;
+		
 		TSharedRef<IConnectionRemapUtils> Utils;
 		TWeakObjectPtr<UVCamWidget> CustomizedWidget;
 		
