Connections named "Selection" do not show up as pin

diff --git a/Engine/Plugins/Experimental/VirtualCamera/Source/VCamBlueprintNodes/Private/SwitchNode/VCamK2Node_SwitchConnectionSystemBase.cpp b/Engine/Plugins/Experimental/VirtualCamera/Source/VCamBlueprintNodes/Private/SwitchNode/VCamK2Node_SwitchConnectionSystemBase.cpp
--- a/Engine/Plugins/Experimental/VirtualCamera/Source/VCamBlueprintNodes/Private/SwitchNode/VCamK2Node_SwitchConnectionSystemBase.cpp
+++ b/Engine/Plugins/Experimental/VirtualCamera/Source/VCamBlueprintNodes/Private/SwitchNode/VCamK2Node_SwitchConnectionSystemBase.cpp
@@ -120,11 +120,11 @@ void UVCamK2Node_SwitchConnectionSystemBase::AddMissingCasePins(const TArray<FNa
 {
 	for (const FName& PinName : PinData)
 	{
-		const bool bContainsPin = Pins.ContainsByPredicate([PinName](UEdGraphPin* Item)
+		const bool bContainsOutputPin = Pins.ContainsByPredicate([PinName](UEdGraphPin* Item)
 		{
-			return Item->PinName == PinName;
+			return Item->PinName == PinName && Item->Direction == EGPD_Output;
 		});
-		if (!bContainsPin)
+		if (!bContainsOutputPin)
 		{
 			CreatePin(EGPD_Output, UEdGraphSchema_K2::PC_Exec, PinName);
 		}
