Add carret to Switch on Connections / Connection Points

diff --git a/Engine/Plugins/Experimental/VirtualCamera/Source/VCamBlueprintNodes/Private/SwitchNode/VCamK2Node_SwitchConnectionSystemBase.cpp b/Engine/Plugins/Experimental/VirtualCamera/Source/VCamBlueprintNodes/Private/SwitchNode/VCamK2Node_SwitchConnectionSystemBase.cpp
--- a/Engine/Plugins/Experimental/VirtualCamera/Source/VCamBlueprintNodes/Private/SwitchNode/VCamK2Node_SwitchConnectionSystemBase.cpp
+++ b/Engine/Plugins/Experimental/VirtualCamera/Source/VCamBlueprintNodes/Private/SwitchNode/VCamK2Node_SwitchConnectionSystemBase.cpp
@@ -14,6 +14,7 @@ UVCamK2Node_SwitchConnectionSystemBase::UVCamK2Node_SwitchConnectionSystemBase()
 {
 	FunctionName = TEXT("NotEqual_NameName");
 	FunctionClass = UKismetMathLibrary::StaticClass();
+	AdvancedPinDisplay = ENodeAdvancedPins::Shown;
 }
 
 void UVCamK2Node_SwitchConnectionSystemBase::RefreshPins()
@@ -38,6 +39,25 @@ void UVCamK2Node_SwitchConnectionSystemBase::PostLoad()
 {
 	Super::PostLoad();
 	SetupBlueprintModifiedCallbacks();
+
+	for (UEdGraphPin* Pin : Pins)
+	{
+		if (Pin->Direction == EGPD_Output)
+		{
+			Pin->bAdvancedView = true;
+		}
+	}
+}
+
+void UVCamK2Node_SwitchConnectionSystemBase::PostEditChangeProperty(FPropertyChangedEvent& PropertyChangedEvent)
+{
+	if (PropertyChangedEvent.MemberProperty
+		&& PropertyChangedEvent.MemberProperty->GetFName() == GET_MEMBER_NAME_CHECKED(UEdGraphNode, AdvancedPinDisplay))
+	{
+		RefreshPins();
+	}
+	
+	Super::PostEditChangeProperty(PropertyChangedEvent);
 }
 
 void UVCamK2Node_SwitchConnectionSystemBase::GetMenuActions(FBlueprintActionDatabaseRegistrar& ActionRegistrar) const
@@ -126,7 +146,8 @@ void UVCamK2Node_SwitchConnectionSystemBase::AddMissingCasePins(const TArray<FNa
 		});
 		if (!bContainsOutputPin)
 		{
-			CreatePin(EGPD_Output, UEdGraphSchema_K2::PC_Exec, PinName);
+			UEdGraphPin* Pin = CreatePin(EGPD_Output, UEdGraphSchema_K2::PC_Exec, PinName);
+			Pin->bAdvancedView = true;
 		}
 	}
 }
diff --git a/Engine/Plugins/Experimental/VirtualCamera/Source/VCamBlueprintNodes/Private/SwitchNode/VCamK2Node_SwitchConnectionSystemBase.h b/Engine/Plugins/Experimental/VirtualCamera/Source/VCamBlueprintNodes/Private/SwitchNode/VCamK2Node_SwitchConnectionSystemBase.h
--- a/Engine/Plugins/Experimental/VirtualCamera/Source/VCamBlueprintNodes/Private/SwitchNode/VCamK2Node_SwitchConnectionSystemBase.h
+++ b/Engine/Plugins/Experimental/VirtualCamera/Source/VCamBlueprintNodes/Private/SwitchNode/VCamK2Node_SwitchConnectionSystemBase.h
@@ -25,6 +25,7 @@ public:
 	//~ Begin UObject Interface
 	virtual void BeginDestroy() override;
 	virtual void PostLoad() override;
+	virtual void PostEditChangeProperty(FPropertyChangedEvent& PropertyChangedEvent) override;
 	//~ End UObject Interface
 	
 	//~ Begin UK2Node Interface
