Detail customization for State Switcher property CurrentState

diff --git a/Engine/Plugins/Experimental/VirtualCamera/Source/VCamCoreEditor/Private/Customizations/OutputProvider/ConnectionRemapCustomization_StateSwitcher.cpp b/Engine/Plugins/Experimental/VirtualCamera/Source/VCamCoreEditor/Private/Customizations/OutputProvider/ConnectionRemapCustomization_StateSwitcher.cpp
--- a/Engine/Plugins/Experimental/VirtualCamera/Source/VCamCoreEditor/Private/Customizations/OutputProvider/ConnectionRemapCustomization_StateSwitcher.cpp
+++ b/Engine/Plugins/Experimental/VirtualCamera/Source/VCamCoreEditor/Private/Customizations/OutputProvider/ConnectionRemapCustomization_StateSwitcher.cpp
@@ -10,6 +10,8 @@
 #include "Algo/AnyOf.h"
 #include "DetailWidgetRow.h"
 #include "IDetailGroup.h"
+#include "Customization/IConnectionRemapUtils.h"
+#include "Util/SharedPropertyCustomizationUtils.h"
 
 #define LOCTEXT_NAMESPACE "FConnectionRemapCustomization_StateSwitcher"
 
@@ -50,7 +52,8 @@ namespace UE::VCamCoreEditor::Private
 			return;
 		}
 
-		Args.WidgetGroup.AddPropertyRow(CurrentStatePropertyHandle.ToSharedRef());
+		IDetailPropertyRow& PropertyRow = Args.WidgetGroup.AddPropertyRow(CurrentStatePropertyHandle.ToSharedRef());
+		StateSwitcher::CustomizeCurrentState(*StateSwitcherWidget, PropertyRow, CurrentStatePropertyHandle.ToSharedRef(), Args.Utils->GetRegularFont());
 	}
 
 }
diff --git a/Engine/Plugins/Experimental/VirtualCamera/Source/VCamCoreEditor/Private/Customizations/StateSwitcher/VCamStateSwitcherWidgetCustomization.cpp b/Engine/Plugins/Experimental/VirtualCamera/Source/VCamCoreEditor/Private/Customizations/StateSwitcher/VCamStateSwitcherWidgetCustomization.cpp
new file mode 100644
--- /dev/null
+++ b/Engine/Plugins/Experimental/VirtualCamera/Source/VCamCoreEditor/Private/Customizations/StateSwitcher/VCamStateSwitcherWidgetCustomization.cpp
@@ -0,0 +1,42 @@
+// Copyright Epic Games, Inc. All Rights Reserved.
+
+#include "VCamStateSwitcherWidgetCustomization.h"
+
+#include "UI/Switcher/VCamStateSwitcherWidget.h"
+
+#include "DetailLayoutBuilder.h"
+#include "DetailWidgetRow.h"
+#include "Util/SharedPropertyCustomizationUtils.h"
+
+namespace UE::VCamCoreEditor::Private
+{
+	TSharedRef<IDetailCustomization> FVCamStateSwitcherWidgetCustomization::MakeInstance()
+	{
+		return MakeShared<FVCamStateSwitcherWidgetCustomization>();
+	}
+
+	void FVCamStateSwitcherWidgetCustomization::CustomizeDetails(IDetailLayoutBuilder& DetailBuilder)
+	{
+		CustomizeCurrentState(DetailBuilder);
+	}
+	
+	void FVCamStateSwitcherWidgetCustomization::CustomizeCurrentState(IDetailLayoutBuilder& DetailBuilder)
+	{
+		const TSharedRef<IPropertyHandle> CurrentStatePropertyHandle = DetailBuilder.GetProperty(UVCamStateSwitcherWidget::GetCurrentStatePropertyName());
+		TArray<TWeakObjectPtr<UObject>> CustomizedObjects;
+		DetailBuilder.GetObjectsBeingCustomized(CustomizedObjects);
+		UVCamStateSwitcherWidget* StateSwitcher = CustomizedObjects.Num() == 1
+			? Cast<UVCamStateSwitcherWidget>(CustomizedObjects[0])
+			: nullptr;
+		if (!StateSwitcher)
+		{
+			return;
+		}
+		
+		IDetailPropertyRow* PropertyRow = DetailBuilder.EditDefaultProperty(CurrentStatePropertyHandle);
+		if (ensure(PropertyRow))
+		{
+			StateSwitcher::CustomizeCurrentState(*StateSwitcher, *PropertyRow, CurrentStatePropertyHandle, DetailBuilder.GetDetailFont());
+		}
+	}
+}
\ No newline at end of file
diff --git a/Engine/Plugins/Experimental/VirtualCamera/Source/VCamCoreEditor/Private/Customizations/StateSwitcher/VCamStateSwitcherWidgetCustomization.h b/Engine/Plugins/Experimental/VirtualCamera/Source/VCamCoreEditor/Private/Customizations/StateSwitcher/VCamStateSwitcherWidgetCustomization.h
new file mode 100644
--- /dev/null
+++ b/Engine/Plugins/Experimental/VirtualCamera/Source/VCamCoreEditor/Private/Customizations/StateSwitcher/VCamStateSwitcherWidgetCustomization.h
@@ -0,0 +1,24 @@
+// Copyright Epic Games, Inc. All Rights Reserved.
+
+#pragma once
+
+#include "CoreMinimal.h"
+#include "IDetailCustomization.h"
+
+namespace UE::VCamCoreEditor::Private
+{
+	/** Makes CurrentState a drop-down based on the available states. */
+	class FVCamStateSwitcherWidgetCustomization : public IDetailCustomization
+	{
+	public:
+		
+		static TSharedRef<IDetailCustomization> MakeInstance();
+		void CustomizeCurrentState(IDetailLayoutBuilder& DetailBuilder);
+
+		//~ Begin IDetailCustomization Interface
+		virtual void CustomizeDetails(IDetailLayoutBuilder& DetailBuilder) override;
+		//~ End IDetailCustomization Interface
+	};
+}
+
+
diff --git a/Engine/Plugins/Experimental/VirtualCamera/Source/VCamCoreEditor/Private/Customizations/StateSwitcher/WidgetConnectionConfigTypeCustomization.cpp b/Engine/Plugins/Experimental/VirtualCamera/Source/VCamCoreEditor/Private/Customizations/StateSwitcher/WidgetConnectionConfigTypeCustomization.cpp
--- a/Engine/Plugins/Experimental/VirtualCamera/Source/VCamCoreEditor/Private/Customizations/StateSwitcher/WidgetConnectionConfigTypeCustomization.cpp
+++ b/Engine/Plugins/Experimental/VirtualCamera/Source/VCamCoreEditor/Private/Customizations/StateSwitcher/WidgetConnectionConfigTypeCustomization.cpp
@@ -160,12 +160,9 @@ namespace UE::VCamCoreEditor::Private
 
 	TObjectPtr<UWidgetTree> FWidgetConnectionConfigTypeCustomization::GetWidgetTreeThroughBlueprintAsset(UUserWidget* ClassDefaultWidget)
 	{
-		if (!ensure(ClassDefaultWidget->HasAnyFlags(RF_ClassDefaultObject)))
-		{
-			return nullptr;
-		}
-		
-		return VCamCore::GetWidgetTreeThroughBlueprintAsset(*ClassDefaultWidget);
+		return ClassDefaultWidget->HasAnyFlags(RF_ClassDefaultObject)
+			? VCamCore::GetWidgetTreeThroughBlueprintAsset(*ClassDefaultWidget)
+			: nullptr;
 	}
 
 	void FWidgetConnectionConfigTypeCustomization::CustomizeConnectionTargetsReferenceProperty(
diff --git a/Engine/Plugins/Experimental/VirtualCamera/Source/VCamCoreEditor/Private/Util/SharedPropertyCustomizationUtils.cpp b/Engine/Plugins/Experimental/VirtualCamera/Source/VCamCoreEditor/Private/Util/SharedPropertyCustomizationUtils.cpp
new file mode 100644
--- /dev/null
+++ b/Engine/Plugins/Experimental/VirtualCamera/Source/VCamCoreEditor/Private/Util/SharedPropertyCustomizationUtils.cpp
@@ -0,0 +1,80 @@
+// Copyright Epic Games, Inc. All Rights Reserved.
+
+#include "SharedPropertyCustomizationUtils.h"
+
+#include "Customizations/StateSwitcher/SStringSelectionComboBox.h"
+#include "UI/Switcher/VCamStateSwitcherWidget.h"
+
+#include "DetailWidgetRow.h"
+#include "IDetailPropertyRow.h"
+#include "Styling/AppStyle.h"
+#include "Widgets/Images/SImage.h"
+#include "Widgets/SBoxPanel.h"
+
+#define LOCTEXT_NAMESPACE "VCamCoreEditor::StateSwitcher"
+
+namespace UE::VCamCoreEditor::Private::StateSwitcher
+{
+	void CustomizeCurrentState(UVCamStateSwitcherWidget& StateSwitcher, IDetailPropertyRow& DetailPropertyRow, TSharedRef<IPropertyHandle> CurrentStatePropertyHandle, const FSlateFontInfo& RegularFont)
+	{
+		TWeakObjectPtr<UVCamStateSwitcherWidget> WeakStateSwitcher = &StateSwitcher;
+		DetailPropertyRow
+			.CustomWidget()
+			.NameContent()
+			[
+				CurrentStatePropertyHandle->CreatePropertyNameWidget()
+			]
+			.ValueContent()
+			[
+				SNew(SHorizontalBox)
+
+				// I don't expect this to ever show since UVCamStateSwitcherWidget SHOULD handle incorrect states
+				+SHorizontalBox::Slot()
+				[
+					SNew(SImage)
+					.ColorAndOpacity(FSlateColor::UseForeground())
+					.DesiredSizeOverride(FVector2D{ 24.f, 24.f })
+					.Image(FAppStyle::Get().GetBrush("Icons.WarningWithColor"))
+					.ToolTipText(LOCTEXT("StateNotFound", "This state was not found on the target widget. Was it removed?"))
+					.Visibility_Lambda([WeakStateSwitcher]()
+					{
+						if (WeakStateSwitcher.IsValid())
+						{
+							return WeakStateSwitcher->GetStates().Contains(WeakStateSwitcher->GetCurrentState())
+								? EVisibility::Collapsed
+								: EVisibility::Visible;
+						}
+						return EVisibility::Collapsed;
+					})
+				]
+
+				+SHorizontalBox::Slot()
+				.AutoWidth()
+				[
+					SNew(SStringSelectionComboBox)
+					.ToolTipText(CurrentStatePropertyHandle->GetToolTipText())
+					.SelectedItem_Lambda([CurrentStatePropertyHandle]()
+					{
+						FName Name;
+						CurrentStatePropertyHandle->GetValue(Name);
+						return Name.ToString();
+					})
+					.ItemList_Lambda([WeakStateSwitcher]() -> TArray<FString>
+					{
+						if (WeakStateSwitcher.IsValid())
+						{
+							TArray<FString> States;
+							Algo::Transform(WeakStateSwitcher->GetStates(), States, [](const FName& Name){ return Name.ToString(); });
+							return States;
+						}
+						return {};
+					})
+					.OnItemSelected_Lambda([CurrentStatePropertyHandle](const FString& SelectedItem)
+					{
+						CurrentStatePropertyHandle->SetValue(FName(*SelectedItem));
+					})
+					.Font(RegularFont)
+				]
+			];
+	}
+}
diff --git a/Engine/Plugins/Experimental/VirtualCamera/Source/VCamCoreEditor/Private/Util/SharedPropertyCustomizationUtils.h b/Engine/Plugins/Experimental/VirtualCamera/Source/VCamCoreEditor/Private/Util/SharedPropertyCustomizationUtils.h
new file mode 100644
--- /dev/null
+++ b/Engine/Plugins/Experimental/VirtualCamera/Source/VCamCoreEditor/Private/Util/SharedPropertyCustomizationUtils.h
@@ -0,0 +1,16 @@
+// Copyright Epic Games, Inc. All Rights Reserved.
+
+#pragma once
+
+#include "CoreMinimal.h"
+
+class IPropertyHandle;
+class IDetailPropertyRow;
+class UVCamStateSwitcherWidget;
+struct FSlateFontInfo;
+
+namespace UE::VCamCoreEditor::Private::StateSwitcher
+{
+	/** Shared logic for customizing UVCamStateSwitcherWidget::CurrentState to display a drop-down widget. */
+	void CustomizeCurrentState(UVCamStateSwitcherWidget& StateSwitcher, IDetailPropertyRow& DetailPropertyRow, TSharedRef<IPropertyHandle> CurrentStatePropertyHandle, const FSlateFontInfo& RegularFont);
+}
diff --git a/Engine/Plugins/Experimental/VirtualCamera/Source/VCamCoreEditor/Private/VCamCoreEditorModule.cpp b/Engine/Plugins/Experimental/VirtualCamera/Source/VCamCoreEditor/Private/VCamCoreEditorModule.cpp
--- a/Engine/Plugins/Experimental/VirtualCamera/Source/VCamCoreEditor/Private/VCamCoreEditorModule.cpp
+++ b/Engine/Plugins/Experimental/VirtualCamera/Source/VCamCoreEditor/Private/VCamCoreEditorModule.cpp
@@ -23,6 +23,7 @@
 #include "Customizations/OutputProvider/ConnectionRemapCustomization_StateSwitcher.h"
 #include "Customizations/OutputProvider/ConnectionRemapCustomization_VCamWidget.h"
 #include "Customizations/OutputProvider/OutputProviderTypeCustomization.h"
+#include "Customizations/StateSwitcher/VCamStateSwitcherWidgetCustomization.h"
 #include "Output/VCamOutputViewport.h"
 #include "UI/VCamWidget.h"
 #include "UI/Switcher/VCamStateSwitcherWidget.h"
@@ -139,6 +140,10 @@ namespace UE::VCamCoreEditor::Private
 			AVCamBaseActor::StaticClass()->GetFName(),
 			FOnGetDetailCustomizationInstance::CreateStatic(&FVCamBaseActorCustomization::MakeInstance)
 		);
+		PropertyModule.RegisterCustomClassLayout(
+			UVCamStateSwitcherWidget::StaticClass()->GetFName(),
+			FOnGetDetailCustomizationInstance::CreateStatic(&FVCamStateSwitcherWidgetCustomization::MakeInstance)
+			);
 	}
 
 	void FVCamCoreEditorModule::RegisterDefaultConnectionRemappingCustomizations()
@@ -165,6 +170,7 @@ namespace UE::VCamCoreEditor::Private
 		
 		PropertyModule.UnregisterCustomClassLayout(UVCamOutputProviderBase::StaticClass()->GetFName());
 		PropertyModule.UnregisterCustomClassLayout(AVCamBaseActor::StaticClass()->GetFName());
+		PropertyModule.UnregisterCustomClassLayout(UVCamStateSwitcherWidget::StaticClass()->GetFName());
 	}
 }
 
