Fix CIS issues

diff --git a/Engine/Plugins/Experimental/VirtualCamera/Source/VCamCoreEditor/Private/Customizations/OutputProvider/ConnectionRemapCustomization_StateSwitcher.cpp b/Engine/Plugins/Experimental/VirtualCamera/Source/VCamCoreEditor/Private/Customizations/OutputProvider/ConnectionRemapCustomization_StateSwitcher.cpp
--- a/Engine/Plugins/Experimental/VirtualCamera/Source/VCamCoreEditor/Private/Customizations/OutputProvider/ConnectionRemapCustomization_StateSwitcher.cpp
+++ b/Engine/Plugins/Experimental/VirtualCamera/Source/VCamCoreEditor/Private/Customizations/OutputProvider/ConnectionRemapCustomization_StateSwitcher.cpp
@@ -26,7 +26,7 @@ namespace UE::VCamCoreEditor::Private
 	{
 		UVCamStateSwitcherWidget* StateSwitcherWidget = Cast<UVCamStateSwitcherWidget>(Args.CustomizedWidget);
 		const bool bHasStateSwitcherProperties = ensure(StateSwitcherWidget)
-			|| Algo::AnyOf(StateSwitcherWidget->GetStates(), [StateSwitcherWidget](FName State)
+			&& Algo::AnyOf(StateSwitcherWidget->GetStates(), [StateSwitcherWidget](FName State)
 			{
 				FVCamWidgetConnectionState StateInfo;
 				return StateSwitcherWidget->GetStateInfo(State, StateInfo)
diff --git a/Engine/Plugins/Experimental/VirtualCamera/Source/VCamCoreEditor/Private/Customizations/OutputProvider/OutputProviderLayoutCustomization.cpp b/Engine/Plugins/Experimental/VirtualCamera/Source/VCamCoreEditor/Private/Customizations/OutputProvider/OutputProviderLayoutCustomization.cpp
--- a/Engine/Plugins/Experimental/VirtualCamera/Source/VCamCoreEditor/Private/Customizations/OutputProvider/OutputProviderLayoutCustomization.cpp
+++ b/Engine/Plugins/Experimental/VirtualCamera/Source/VCamCoreEditor/Private/Customizations/OutputProvider/OutputProviderLayoutCustomization.cpp
@@ -253,7 +253,7 @@ namespace UE::VCamCoreEditor::Private
 		const UVPFullScreenUserWidget* FullScreenUserWidget = CustomizedOutputProvider->GetUMGWidget();
 		UUserWidget* RootWidget = FullScreenUserWidget ? FullScreenUserWidget->GetWidget() : nullptr;
 		UWidgetTree* WidgetTree = RootWidget ? RootWidget->WidgetTree : nullptr;
-		if (!WidgetTree)
+		if (!WidgetTree || !RootWidget) // !RootWidget pointless but maks static analyzer happy
 		{
 			return;
 		}
diff --git a/Engine/Plugins/Experimental/VirtualCamera/Source/VCamCoreEditor/Private/Customizations/VCamViewportLockerTypeCustomization.cpp b/Engine/Plugins/Experimental/VirtualCamera/Source/VCamCoreEditor/Private/Customizations/VCamViewportLockerTypeCustomization.cpp
--- a/Engine/Plugins/Experimental/VirtualCamera/Source/VCamCoreEditor/Private/Customizations/VCamViewportLockerTypeCustomization.cpp
+++ b/Engine/Plugins/Experimental/VirtualCamera/Source/VCamCoreEditor/Private/Customizations/VCamViewportLockerTypeCustomization.cpp
@@ -16,9 +16,9 @@ namespace UE::VCamCoreEditor::Private
 		switch (ViewportIndex)
 		{
 		case 0: return LOCTEXT("Viewport1", "Viewport 1");
-		case 1: return LOCTEXT("Viewport1", "Viewport 2");
-		case 2: return LOCTEXT("Viewport1", "Viewport 3");
-		case 3: return LOCTEXT("Viewport1", "Viewport 4");
+		case 1: return LOCTEXT("Viewport2", "Viewport 2");
+		case 2: return LOCTEXT("Viewport3", "Viewport 3");
+		case 3: return LOCTEXT("Viewport4", "Viewport 4");
 
 		default:
 			checkNoEntry();
