Fix missing #undef LOCTEXT_NAMESPACE

diff --git a/Engine/Plugins/Experimental/VirtualCamera/Source/VCamCoreEditor/Private/Util/SharedPropertyCustomizationUtils.cpp b/Engine/Plugins/Experimental/VirtualCamera/Source/VCamCoreEditor/Private/Util/SharedPropertyCustomizationUtils.cpp
--- a/Engine/Plugins/Experimental/VirtualCamera/Source/VCamCoreEditor/Private/Util/SharedPropertyCustomizationUtils.cpp
+++ b/Engine/Plugins/Experimental/VirtualCamera/Source/VCamCoreEditor/Private/Util/SharedPropertyCustomizationUtils.cpp
@@ -78,3 +78,5 @@ namespace UE::VCamCoreEditor::Private::StateSwitcher
 			];
 	}
 }
+
+#undef LOCTEXT_NAMESPACE
\ No newline at end of file
