Added FWidgetReferenceForBlueprint

diff --git a/Engine/Plugins/Experimental/VirtualCamera/Source/VCamCore/Private/Util/WidgetReference.cpp b/Engine/Plugins/Experimental/VirtualCamera/Source/VCamCore/Private/Util/WidgetReference.cpp
new file mode 100644
--- /dev/null
+++ b/Engine/Plugins/Experimental/VirtualCamera/Source/VCamCore/Private/Util/WidgetReference.cpp
@@ -0,0 +1,34 @@
+// Copyright Epic Games, Inc. All Rights Reserved.
+
+#include "Util/WidgetReference.h"
+
+#include "LogVCamCore.h"
+#include "UI/Switcher/VCamStateSwitcherWidget.h"
+#include "Util/WidgetTreeUtils.h"
+
+#include "Blueprint/WidgetTree.h"
+
+UWidget* FChildWidgetReference::ResolveWidget(UUserWidget& OwnerWidget) const
+{
+	UWidget* ResolvedTemplate = Template.Get();
+	if (!ResolvedTemplate)
+	{
+		return nullptr;
+	}
+	
+	UWidgetTree* WidgetTree = OwnerWidget.WidgetTree;
+#if WITH_EDITOR
+	if (!WidgetTree && OwnerWidget.HasAnyFlags(RF_ClassDefaultObject))
+	{
+		WidgetTree = UE::VCamCore::GetWidgetTreeThroughBlueprintAsset(OwnerWidget);
+	}
+#endif
+
+	const bool bCanFindWidget = WidgetTree != nullptr;;
+	UE_CLOG(!bCanFindWidget, LogVCamCore, Warning, TEXT("Failed to get tree for widget %s"), *OwnerWidget.GetPathName());
+
+	const FName WidgetName = ResolvedTemplate->GetFName();
+	return bCanFindWidget
+		? WidgetTree->FindWidget<UWidget>(WidgetName)
+		: nullptr;
+}
diff --git a/Engine/Plugins/Experimental/VirtualCamera/Source/VCamCore/Public/Util/WidgetReference.h b/Engine/Plugins/Experimental/VirtualCamera/Source/VCamCore/Public/Util/WidgetReference.h
new file mode 100644
--- /dev/null
+++ b/Engine/Plugins/Experimental/VirtualCamera/Source/VCamCore/Public/Util/WidgetReference.h
@@ -0,0 +1,59 @@
+// Copyright Epic Games, Inc. All Rights Reserved.
+
+#pragma once
+
+#include "CoreMinimal.h"
+#include "Kismet/BlueprintFunctionLibrary.h"
+#include "UI/VCamWidget.h"
+#include "WidgetReference.generated.h"
+
+/**
+ * A reference to any child widget within an UMG Blueprint; only widgets in the same widget tree can be referenced.
+ * Has a detail customization for choosing widget.
+ */
+USTRUCT(BlueprintType)
+struct VCAMCORE_API FChildWidgetReference
+{
+	GENERATED_BODY()
+
+	/**
+	 * Pointer to template widget within the Blueprint's widget tree.
+	 * 
+	 * Important: in UMG there are templates and previews.
+	 * Templates are like the CDO and previews are what you see in the editor
+	 * @see FWidgetReference
+	 */
+	UPROPERTY(EditAnywhere, Category = "Virtual Camera")
+	TLazyObjectPtr<UWidget> Template;
+	
+	UWidget* ResolveWidget(UUserWidget& OwnerWidget) const;
+	bool HasNoWidgetSet() const { return Template.IsValid(); }
+};
+
+/** Version that restraints Widget to UVCamWidget instances. */
+USTRUCT(BlueprintType)
+struct VCAMCORE_API FVCamChildWidgetReference : public FChildWidgetReference
+{
+	GENERATED_BODY()
+	
+	UVCamWidget* ResolveVCamWidget(UUserWidget& OwnerWidget) const { return Cast<UVCamWidget>(ResolveWidget(OwnerWidget)); }
+};
+
+UCLASS()
+class UWidgetReferenceBlueprintFunctionLibrary : public UBlueprintFunctionLibrary
+{
+	GENERATED_BODY()
+public:
+
+	UFUNCTION(BlueprintPure, Category = "Virtual Camera")
+	static UWidget* ResolveWidget(const FChildWidgetReference& WidgetReference, UUserWidget* OwnerWidget)
+	{
+		return OwnerWidget ? WidgetReference.ResolveWidget(*OwnerWidget) : nullptr;
+	}
+	
+	UFUNCTION(BlueprintPure, Category = "Virtual Camera")
+	static UVCamWidget* ResolveVCamWidget(const FVCamChildWidgetReference& WidgetReference, UUserWidget* OwnerWidget)
+	{
+		return OwnerWidget ? WidgetReference.ResolveVCamWidget(*OwnerWidget) : nullptr;
+	}
+};
\ No newline at end of file
diff --git a/Engine/Plugins/Experimental/VirtualCamera/Source/VCamCoreEditor/Private/Customizations/WidgetReference/ChildWidgetReferenceCustomization.cpp b/Engine/Plugins/Experimental/VirtualCamera/Source/VCamCoreEditor/Private/Customizations/WidgetReference/ChildWidgetReferenceCustomization.cpp
new file mode 100644
--- /dev/null
+++ b/Engine/Plugins/Experimental/VirtualCamera/Source/VCamCoreEditor/Private/Customizations/WidgetReference/ChildWidgetReferenceCustomization.cpp
@@ -0,0 +1,130 @@
+// Copyright Epic Games, Inc. All Rights Reserved.
+
+#include "ChildWidgetReferenceCustomization.h"
+
+#include "Customizations/StateSwitcher/SStringSelectionComboBox.h"
+#include "Util/WidgetReference.h"
+#include "Util/WidgetTreeUtils.h"
+
+#include "Algo/AnyOf.h"
+#include "Blueprint/UserWidget.h"
+#include "Blueprint/WidgetTree.h"
+#include "DetailWidgetRow.h"
+#include "IPropertyUtilities.h"
+#include "Widgets/Text/STextBlock.h"
+
+#define LOCTEXT_NAMESPACE "FWidgetReferenceForBlueprintCustomization"
+
+namespace UE::VCamCoreEditor::Private
+{
+	TSharedRef<IPropertyTypeCustomization> FChildWidgetReferenceCustomization::MakeInstance()
+	{
+		return MakeShared<FChildWidgetReferenceCustomization>();
+	}
+
+	void FChildWidgetReferenceCustomization::CustomizeHeader(TSharedRef<IPropertyHandle> PropertyHandle, FDetailWidgetRow& HeaderRow, IPropertyTypeCustomizationUtils& CustomizationUtils)
+	{
+		const TSharedPtr<IPropertyHandle> TemplateProperty = PropertyHandle->GetChildHandle(GET_MEMBER_NAME_CHECKED(FChildWidgetReference, Template));
+		
+		const TArray<TWeakObjectPtr<UObject>>& EditedObjects = CustomizationUtils.GetPropertyUtilities()->GetSelectedObjects();
+		TWeakObjectPtr<UUserWidget> EditedWidget = EditedObjects.Num() == 1 ? Cast<UUserWidget>(EditedObjects[0]) : nullptr;
+		// Always only search on the template
+		if (EditedWidget.IsValid() && !EditedWidget->HasAnyFlags(RF_ClassDefaultObject))
+		{
+			UObject* DefaultObject = EditedWidget->GetClass()->GetDefaultObject();
+			EditedWidget = Cast<UUserWidget>(DefaultObject);
+		}
+		
+		if (!EditedWidget.IsValid() || !ensure(TemplateProperty))
+		{
+			const FText Text = LOCTEXT("NotEditable", "No editable in this context");
+			HeaderRow.
+				NameContent()
+				[
+					PropertyHandle->CreatePropertyNameWidget()
+				]
+				.ValueContent()
+				[
+					SNew(STextBlock)
+					.Text(Text)
+					.ToolTipText(Text)
+					.Font(CustomizationUtils.GetRegularFont())
+				];
+			return;
+		}
+		
+		HeaderRow.
+			NameContent()
+			[
+				PropertyHandle->CreatePropertyNameWidget()
+			]
+			.ValueContent()
+			[
+				SNew(SStringSelectionComboBox)
+				.ToolTipText(LOCTEXT("PropertySelection.Tooltip", "Select a child widget from this Blueprint.\n\nThis can either be:\n\t- an auto-generated variable (see \"Is Variable\" check box in Designer)\n\t- a regular object property of type VCamWidget (you must make sure the property points to something valid before you change states)"))
+				.SelectedItem_Lambda([TemplateProperty]()
+				{
+					UObject* Object = nullptr;
+					TemplateProperty->GetValue(Object);
+					return Object ? Object->GetName() : TEXT("None");
+				})
+				// Property list is only build once: when the Blueprint is recompiled,
+				// the details view SHOULD refresh automatically and thus all of this is reconstructed. 
+				.ItemList(this, &FChildWidgetReferenceCustomization::GetPropertyItemList, EditedWidget)
+				.OnItemSelected_Lambda([this, TemplateProperty, EditedWidget](const FString& SelectedItem)
+				{
+					const TArray<UWidget*> Widgets = GetSelectableChildWidgets(EditedWidget);
+					const FName WidgetName = *SelectedItem;
+					UWidget* const* FoundWidget = Widgets.FindByPredicate([&WidgetName](UWidget* Widget){ return Widget->GetFName() == WidgetName; });
+					
+					// Setting null value is allowed (and even required)
+					const FPropertyAccess::Result AccessResult = TemplateProperty->SetValue(FoundWidget ? *FoundWidget : nullptr);
+					ensure(AccessResult == FPropertyAccess::Success);
+				})
+				.Font(CustomizationUtils.GetRegularFont())
+			];
+	}
+
+	TArray<UWidget*> FChildWidgetReferenceCustomization::GetSelectableChildWidgets(TWeakObjectPtr<UUserWidget> Widget) const
+	{
+		if (!Widget.IsValid())
+		{
+			return {};
+		}
+		
+		const UWidgetTree* WidgetTree = Widget->WidgetTree
+			? Widget->WidgetTree.Get()
+			: VCamCore::GetWidgetTreeThroughBlueprintAsset(*Widget);
+		if (!WidgetTree)
+		{
+			return {};
+		}
+		
+		TArray<UWidget*> Widgets;
+		WidgetTree->ForEachWidget([&Widgets](UWidget* Widget)
+		{
+			if (UWidget* UserWidget = Cast<UWidget>(Widget))
+			{
+				Widgets.Add(UserWidget);
+			}
+		});
+		return Widgets;
+	}
+
+	TArray<FString> FChildWidgetReferenceCustomization::GetPropertyItemList(TWeakObjectPtr<UUserWidget> Widget) const
+	{
+		TArray<FString> WidgetNamesAsStrings;
+		Algo::Transform(GetSelectableChildWidgets(Widget), WidgetNamesAsStrings, [](UWidget* Widget){ return Widget->GetName(); });
+		if (WidgetNamesAsStrings.Num() == 0)
+		{
+			WidgetNamesAsStrings.Add(TEXT("None"));
+		}
+		else
+		{
+			WidgetNamesAsStrings.Insert(TEXT("None"), 0);
+		}
+		return WidgetNamesAsStrings;
+	}
+}
+
+#undef LOCTEXT_NAMESPACE
diff --git a/Engine/Plugins/Experimental/VirtualCamera/Source/VCamCoreEditor/Private/Customizations/WidgetReference/ChildWidgetReferenceCustomization.h b/Engine/Plugins/Experimental/VirtualCamera/Source/VCamCoreEditor/Private/Customizations/WidgetReference/ChildWidgetReferenceCustomization.h
new file mode 100644
--- /dev/null
+++ b/Engine/Plugins/Experimental/VirtualCamera/Source/VCamCoreEditor/Private/Customizations/WidgetReference/ChildWidgetReferenceCustomization.h
@@ -0,0 +1,34 @@
+// Copyright Epic Games, Inc. All Rights Reserved.
+
+#pragma once
+
+#include "CoreMinimal.h"
+#include "IPropertyTypeCustomization.h"
+
+class UUserWidget;
+class UWidget;
+
+namespace UE::VCamCoreEditor::Private
+{
+	class FChildWidgetReferenceCustomization : public IPropertyTypeCustomization
+	{
+	public:
+
+		static TSharedRef<IPropertyTypeCustomization> MakeInstance();
+		
+		//~ Begin IDetailCustomization Interface
+		virtual void CustomizeHeader(TSharedRef<IPropertyHandle> PropertyHandle, FDetailWidgetRow& HeaderRow, IPropertyTypeCustomizationUtils& CustomizationUtils) override;
+		virtual void CustomizeChildren(TSharedRef<IPropertyHandle> PropertyHandle, IDetailChildrenBuilder& ChildBuilder, IPropertyTypeCustomizationUtils& CustomizationUtils) override {}
+		//~ End IDetailCustomization Interface
+
+	protected:
+		
+		virtual TArray<UWidget*> GetSelectableChildWidgets(TWeakObjectPtr<UUserWidget> Widget) const;
+
+	private:
+		
+		TArray<FString> GetPropertyItemList(TWeakObjectPtr<UUserWidget> Widget) const;
+	};
+}
+
+
diff --git a/Engine/Plugins/Experimental/VirtualCamera/Source/VCamCoreEditor/Private/Customizations/WidgetReference/VCamChildWidgetReferenceCustomization.cpp b/Engine/Plugins/Experimental/VirtualCamera/Source/VCamCoreEditor/Private/Customizations/WidgetReference/VCamChildWidgetReferenceCustomization.cpp
new file mode 100644
--- /dev/null
+++ b/Engine/Plugins/Experimental/VirtualCamera/Source/VCamCoreEditor/Private/Customizations/WidgetReference/VCamChildWidgetReferenceCustomization.cpp
@@ -0,0 +1,22 @@
+// Copyright Epic Games, Inc. All Rights Reserved.
+
+#include "VCamChildWidgetReferenceCustomization.h"
+
+#include "Algo/RemoveIf.h"
+#include "Blueprint/UserWidget.h"
+#include "UI/VCamWidget.h"
+
+namespace UE::VCamCoreEditor::Private
+{
+	TSharedRef<IPropertyTypeCustomization> FVCamChildWidgetReferenceCustomization::MakeInstance()
+	{
+		return MakeShared<FVCamChildWidgetReferenceCustomization>();
+	}
+
+	TArray<UWidget*> FVCamChildWidgetReferenceCustomization::GetSelectableChildWidgets(TWeakObjectPtr<UUserWidget> Widget) const
+	{
+		TArray<UWidget*> Widgets = FChildWidgetReferenceCustomization::GetSelectableChildWidgets(Widget);
+		Widgets.SetNum(Algo::RemoveIf(Widgets, [](UWidget* Widget){ return Cast<UVCamWidget>(Widget) == nullptr; }));
+		return Widgets;
+	}
+}
diff --git a/Engine/Plugins/Experimental/VirtualCamera/Source/VCamCoreEditor/Private/Customizations/WidgetReference/VCamChildWidgetReferenceCustomization.h b/Engine/Plugins/Experimental/VirtualCamera/Source/VCamCoreEditor/Private/Customizations/WidgetReference/VCamChildWidgetReferenceCustomization.h
new file mode 100644
--- /dev/null
+++ b/Engine/Plugins/Experimental/VirtualCamera/Source/VCamCoreEditor/Private/Customizations/WidgetReference/VCamChildWidgetReferenceCustomization.h
@@ -0,0 +1,26 @@
+// Copyright Epic Games, Inc. All Rights Reserved.
+
+#pragma once
+
+#include "CoreMinimal.h"
+#include "ChildWidgetReferenceCustomization.h"
+
+class UUserWidget;
+
+namespace UE::VCamCoreEditor::Private
+{
+	class FVCamChildWidgetReferenceCustomization : public FChildWidgetReferenceCustomization
+	{
+	public:
+		
+		static TSharedRef<IPropertyTypeCustomization> MakeInstance();
+		
+	protected:
+		
+		//~ Begin FWidgetReferenceForBlueprintCustomization Interface
+		virtual TArray<UWidget*> GetSelectableChildWidgets(TWeakObjectPtr<UUserWidget> Widget) const override;
+		//~ End FWidgetReferenceForBlueprintCustomization Interface
+	};
+}
+
+
diff --git a/Engine/Plugins/Experimental/VirtualCamera/Source/VCamCoreEditor/Private/VCamCoreEditorModule.cpp b/Engine/Plugins/Experimental/VirtualCamera/Source/VCamCoreEditor/Private/VCamCoreEditorModule.cpp
--- a/Engine/Plugins/Experimental/VirtualCamera/Source/VCamCoreEditor/Private/VCamCoreEditorModule.cpp
+++ b/Engine/Plugins/Experimental/VirtualCamera/Source/VCamCoreEditor/Private/VCamCoreEditorModule.cpp
@@ -3,30 +3,32 @@
 #include "VCamCoreEditorModule.h"
 
 #include "Customizations/ConnectionTargetSettingsTypeCustomization.h"
+#include "Customizations/OutputProvider/ConnectionRemapCustomization_StateSwitcher.h"
+#include "Customizations/OutputProvider/ConnectionRemapCustomization_VCamWidget.h"
 #include "Customizations/OutputProvider/OutputProviderLayoutCustomization.h"
+#include "Customizations/OutputProvider/OutputProviderTypeCustomization.h"
+#include "Customizations/StateSwitcher/VCamStateSwitcherWidgetCustomization.h"
+#include "Customizations/StateSwitcher/WidgetConnectionConfigTypeCustomization.h"
 #include "Customizations/VCamBaseActorCustomization.h"
 #include "Customizations/VCamInputProfileCustomization.h"
-#include "Customizations/StateSwitcher/WidgetConnectionConfigTypeCustomization.h"
+#include "Customizations/VCamViewportLockerTypeCustomization.h"
+#include "Customizations/WidgetReference/ChildWidgetReferenceCustomization.h"
+#include "Customizations/WidgetReference/VCamChildWidgetReferenceCustomization.h"
+#include "LogVCamEditor.h"
+#include "Modifier/VCamModifier.h"
+#include "Output/VCamOutputProviderBase.h"
 #include "UI/Switcher/WidgetConnectionConfig.h"
+#include "UI/Switcher/VCamStateSwitcherWidget.h"
+#include "UI/VCamWidget.h"
+#include "Util/WidgetReference.h"
 #include "VCamBaseActor.h"
 #include "VCamComponent.h"
-#include "Output/VCamOutputProviderBase.h"
-#include "LogVCamEditor.h"
-#include "Modifier/VCamModifier.h"
 
 #include "Interfaces/IPluginManager.h"
 #include "Kismet2/KismetEditorUtilities.h"
 #include "Modules/ModuleInterface.h"
 #include "Modules/ModuleManager.h"
 #include "PropertyEditorModule.h"
-#include "Customizations/VCamViewportLockerTypeCustomization.h"
-#include "Customizations/OutputProvider/ConnectionRemapCustomization_StateSwitcher.h"
-#include "Customizations/OutputProvider/ConnectionRemapCustomization_VCamWidget.h"
-#include "Customizations/OutputProvider/OutputProviderTypeCustomization.h"
-#include "Customizations/StateSwitcher/VCamStateSwitcherWidgetCustomization.h"
-#include "Output/VCamOutputViewport.h"
-#include "UI/VCamWidget.h"
-#include "UI/Switcher/VCamStateSwitcherWidget.h"
 
 namespace UE::VCamCoreEditor::Private
 {
@@ -130,7 +132,15 @@ namespace UE::VCamCoreEditor::Private
 		PropertyModule.RegisterCustomPropertyTypeLayout(
 			FVCamViewportLocker::StaticStruct()->GetFName(), 
 			FOnGetPropertyTypeCustomizationInstance::CreateStatic(&FVCamViewportLockerTypeCustomization::MakeInstance)
-		);
+			);
+		PropertyModule.RegisterCustomPropertyTypeLayout(
+			FChildWidgetReference::StaticStruct()->GetFName(),
+			FOnGetPropertyTypeCustomizationInstance::CreateStatic(&FChildWidgetReferenceCustomization::MakeInstance)
+			);
+		PropertyModule.RegisterCustomPropertyTypeLayout(
+			FVCamChildWidgetReference::StaticStruct()->GetFName(),
+			FOnGetPropertyTypeCustomizationInstance::CreateStatic(&FVCamChildWidgetReferenceCustomization::MakeInstance)
+			);
 		
 		PropertyModule.RegisterCustomClassLayout(
 			UVCamOutputProviderBase::StaticClass()->GetFName(),
@@ -167,6 +177,8 @@ namespace UE::VCamCoreEditor::Private
 		PropertyModule.UnregisterCustomPropertyTypeLayout(FWidgetConnectionConfig::StaticStruct()->GetFName());
 		PropertyModule.UnregisterCustomPropertyTypeLayout(UVCamOutputProviderBase::StaticClass()->GetFName());
 		PropertyModule.UnregisterCustomPropertyTypeLayout(FVCamViewportLocker::StaticStruct()->GetFName());
+		PropertyModule.UnregisterCustomPropertyTypeLayout(FChildWidgetReference::StaticStruct()->GetFName());
+		PropertyModule.UnregisterCustomPropertyTypeLayout(FVCamChildWidgetReference::StaticStruct()->GetFName());
 		
 		PropertyModule.UnregisterCustomClassLayout(UVCamOutputProviderBase::StaticClass()->GetFName());
 		PropertyModule.UnregisterCustomClassLayout(AVCamBaseActor::StaticClass()->GetFName());
