Widget references in state switcher are "lost" when widget is renamed

diff --git a/Engine/Plugins/Experimental/VirtualCamera/Source/VCamCore/Private/UI/Switcher/VCamStateSwitcherWidget.cpp b/Engine/Plugins/Experimental/VirtualCamera/Source/VCamCore/Private/UI/Switcher/VCamStateSwitcherWidget.cpp
--- a/Engine/Plugins/Experimental/VirtualCamera/Source/VCamCore/Private/UI/Switcher/VCamStateSwitcherWidget.cpp
+++ b/Engine/Plugins/Experimental/VirtualCamera/Source/VCamCore/Private/UI/Switcher/VCamStateSwitcherWidget.cpp
@@ -25,12 +25,13 @@ bool UVCamStateSwitcherWidget::SetCurrentState(FName NewState, bool bForceUpdate
 	const FName OldState = CurrentState;
 	OnPreStateChanged.Broadcast(this, OldState, NewState);
 	
-	for (const FWidgetConnectionConfig& WidgetConfig : StateConfig->WidgetConfigs)
+	for (int32 i = 0; i < StateConfig->WidgetConfigs.Num(); ++i)
 	{
+		const FWidgetConnectionConfig& WidgetConfig = StateConfig->WidgetConfigs[i];
 		UVCamWidget* Widget = WidgetConfig.ResolveWidget(this);
 		if (!Widget)
 		{
-			UE_CLOG(!WidgetConfig.HasNoWidgetSet(), LogVCamStateSwitcher, Warning, TEXT("Failed to find widget %s (NewState: %s)"), *WidgetConfig.Widget.ToString(), *NewState.ToString());
+			UE_CLOG(!WidgetConfig.HasNoWidgetSet(), LogVCamStateSwitcher, Warning, TEXT("Failed to find widget at index %d in state %s"), i, *NewState.ToString());
 			continue;
 		}
 
diff --git a/Engine/Plugins/Experimental/VirtualCamera/Source/VCamCore/Private/UI/Switcher/WidgetConnectionConfig.cpp b/Engine/Plugins/Experimental/VirtualCamera/Source/VCamCore/Private/UI/Switcher/WidgetConnectionConfig.cpp
--- a/Engine/Plugins/Experimental/VirtualCamera/Source/VCamCore/Private/UI/Switcher/WidgetConnectionConfig.cpp
+++ b/Engine/Plugins/Experimental/VirtualCamera/Source/VCamCore/Private/UI/Switcher/WidgetConnectionConfig.cpp
@@ -2,30 +2,18 @@
 
 #include "UI/Switcher/WidgetConnectionConfig.h"
 
-#include "LogVCamCore.h"
 #include "UI/Switcher/VCamStateSwitcherWidget.h"
-#include "Util/WidgetTreeUtils.h"
 
 #include "Blueprint/WidgetTree.h"
 
 UVCamWidget* FWidgetConnectionConfig::ResolveWidget(UVCamStateSwitcherWidget* OwnerWidget) const
 {
-	if (!OwnerWidget)
-	{
-		return nullptr;
-	}
-	
-	UWidgetTree* WidgetTree = OwnerWidget->WidgetTree;
-#if WITH_EDITOR
-	if (!WidgetTree && OwnerWidget->HasAnyFlags(RF_ClassDefaultObject))
-	{
-		WidgetTree = UE::VCamCore::GetWidgetTreeThroughBlueprintAsset(*OwnerWidget);
-	}
-#endif
-
-	const bool bCanFindWidget = !HasNoWidgetSet() && WidgetTree;
-	UE_CLOG(!bCanFindWidget, LogVCamCore, Warning, TEXT("Failed to get tree for widget %s"), *OwnerWidget->GetPathName());
-	return bCanFindWidget
-		? WidgetTree->FindWidget<UVCamWidget>(Widget)
+	return OwnerWidget
+		? Widget.ResolveVCamWidget(*OwnerWidget)
 		: nullptr;
 }
+
+bool FWidgetConnectionConfig::HasNoWidgetSet() const
+{
+	return Widget.HasNoWidgetSet();
+}
diff --git a/Engine/Plugins/Experimental/VirtualCamera/Source/VCamCore/Public/UI/Switcher/WidgetConnectionConfig.h b/Engine/Plugins/Experimental/VirtualCamera/Source/VCamCore/Public/UI/Switcher/WidgetConnectionConfig.h
--- a/Engine/Plugins/Experimental/VirtualCamera/Source/VCamCore/Public/UI/Switcher/WidgetConnectionConfig.h
+++ b/Engine/Plugins/Experimental/VirtualCamera/Source/VCamCore/Public/UI/Switcher/WidgetConnectionConfig.h
@@ -4,6 +4,7 @@
 
 #include "CoreMinimal.h"
 #include "UI/VCamConnectionStructs.h"
+#include "Util/WidgetReference.h"
 #include "WidgetConnectionConfig.generated.h"
 
 class UVCamStateSwitcherWidget;
@@ -19,7 +20,7 @@ struct VCAMCORE_API FWidgetConnectionConfig
 	 * This is the name of a child widget and can be used as argument to UWidgetTree::FindWidget.
 	 */
 	UPROPERTY(EditAnywhere, BlueprintReadWrite, Category = "Connection")
-	FName Widget = NAME_None;
+	FVCamChildWidgetReference Widget;
 	
 	/**
 	 * Key: A valid key of in UVCamWidget::Connections of the widget identified via WidgetProperty.
@@ -29,5 +30,5 @@ struct VCAMCORE_API FWidgetConnectionConfig
 	TMap<FName, FVCamConnectionTargetSettings> ConnectionTargets;
 
 	UVCamWidget* ResolveWidget(UVCamStateSwitcherWidget* OwnerWidget) const;
-	bool HasNoWidgetSet() const { return Widget == NAME_None; }
+	bool HasNoWidgetSet() const;
 };
\ No newline at end of file
diff --git a/Engine/Plugins/Experimental/VirtualCamera/Source/VCamCoreEditor/Private/Customizations/StateSwitcher/WidgetConnectionConfigTypeCustomization.cpp b/Engine/Plugins/Experimental/VirtualCamera/Source/VCamCoreEditor/Private/Customizations/StateSwitcher/WidgetConnectionConfigTypeCustomization.cpp
--- a/Engine/Plugins/Experimental/VirtualCamera/Source/VCamCoreEditor/Private/Customizations/StateSwitcher/WidgetConnectionConfigTypeCustomization.cpp
+++ b/Engine/Plugins/Experimental/VirtualCamera/Source/VCamCoreEditor/Private/Customizations/StateSwitcher/WidgetConnectionConfigTypeCustomization.cpp
@@ -69,102 +69,10 @@ namespace UE::VCamCoreEditor::Private
 			}
 		}
 		
-		IDetailPropertyRow& ReferenceRow = ChildBuilder.AddProperty(WidgetReferencePropertyHandle.ToSharedRef());
-		CustomizeWidgetReferenceProperty(WidgetReferencePropertyHandle.ToSharedRef(), ReferenceRow, CustomizationUtils);
-		
+		ChildBuilder.AddProperty(WidgetReferencePropertyHandle.ToSharedRef());
 		CustomizeConnectionTargetsReferenceProperty(StructPropertyHandle, ConnectionTargetsPropertyHandle.ToSharedRef(), ChildBuilder, CustomizationUtils);
 	}
 
-	void FWidgetConnectionConfigTypeCustomization::CustomizeWidgetReferenceProperty(
-		TSharedRef<IPropertyHandle> WidgetReferencePropertyHandle,
-		IDetailPropertyRow& Row,
-		IPropertyTypeCustomizationUtils& CustomizationUtils) const
-	{
-		WidgetReferencePropertyHandle->SetOnPropertyValueChanged(FSimpleDelegate::CreateLambda([PropertyUtils = CustomizationUtils.GetPropertyUtilities()]()
-		{
-			PropertyUtils->ForceRefresh();
-		}));
-
-		const TWeakPtr<IPropertyUtilities> WeakPropertyUtils = CustomizationUtils.GetPropertyUtilities();
-		Row.CustomWidget()
-			.NameContent()
-			[
-				WidgetReferencePropertyHandle->CreatePropertyNameWidget()
-			]
-			.ValueContent()
-			[
-				SNew(SStringSelectionComboBox)
-					.ToolTipText(LOCTEXT("PropertySelection.Tooltip", "Select a child widget from this Blueprint.\n\nThis can either be:\n\t- an auto-generated variable (see \"Is Variable\" check box in Designer)\n\t- a regular object property of type VCamWidget (you must make sure the property points to something valid before you change states)"))
-					.SelectedItem_Lambda([WidgetReferencePropertyHandle]()
-					{
-						FName Name;
-						WidgetReferencePropertyHandle->GetValue(Name);
-						return Name.ToString();
-					})
-					// Property list is only build once: when the Blueprint is recompiled,
-					// the details view SHOULD refresh automatically and thus all of this is reconstructed. 
-					.ItemList(this, &FWidgetConnectionConfigTypeCustomization::GetPropertyItemList, WeakPropertyUtils)
-					.OnItemSelected_Lambda([WidgetReferencePropertyHandle](const FString& SelectedItem)
-					{
-						WidgetReferencePropertyHandle->SetValue(SelectedItem);
-					})
-					.Font(CustomizationUtils.GetRegularFont())
-			];
-	}
-	
-	TArray<FString> FWidgetConnectionConfigTypeCustomization::GetPropertyItemList(TWeakPtr<IPropertyUtilities> WeakPropertyUtils) const
-	{
-		TSharedPtr<IPropertyUtilities> PropertyUtilities = WeakPropertyUtils.Pin();
-		if (!PropertyUtilities)
-		{
-			return {};
-		}
-		
-		TArray<FString> PropertyItemList;
-		PropertyItemList.Add(FName(NAME_None).ToString());
-		
-		const TArray<TWeakObjectPtr<UObject>>& SelectedObjects = PropertyUtilities->GetSelectedObjects();
-		if (!ensure(SelectedObjects.Num() == 1) || !SelectedObjects[0].IsValid())
-		{
-			return PropertyItemList;
-		}
-
-		UObject* Object = SelectedObjects[0].Get();
-		UUserWidget* Widget = Cast<UUserWidget>(Object);
-		if (!Widget)
-		{
-			return {};
-		}
-
-		// Two use cases are expected: Editing an instance from 1. UMG Designer tab 2. Class Defaults in Graph tab
-		// When editing Class Defaults, the WidgetTree is expected to be nullptr when editing class defaults so let's grab it from the asset.
-		const TObjectPtr<UWidgetTree> WidgetTree = Widget->WidgetTree
-			? Widget->WidgetTree
-			: GetWidgetTreeThroughBlueprintAsset(Widget);
-		if (!WidgetTree)
-		{
-			return {};
-		}
-		
-		WidgetTree->ForEachWidget([&PropertyItemList](UWidget* Widget)
-		{
-			const bool bIsVCamWidgetProperty = Widget && Widget->GetClass()->IsChildOf(UVCamWidget::StaticClass());
-			if (bIsVCamWidgetProperty)
-			{
-				PropertyItemList.Add(Widget->GetName());
-			}
-		});
-		
-		return PropertyItemList;
-	}
-
-	TObjectPtr<UWidgetTree> FWidgetConnectionConfigTypeCustomization::GetWidgetTreeThroughBlueprintAsset(UUserWidget* ClassDefaultWidget)
-	{
-		return ClassDefaultWidget->HasAnyFlags(RF_ClassDefaultObject)
-			? VCamCore::GetWidgetTreeThroughBlueprintAsset(*ClassDefaultWidget)
-			: nullptr;
-	}
-
 	void FWidgetConnectionConfigTypeCustomization::CustomizeConnectionTargetsReferenceProperty(
 		TSharedRef<IPropertyHandle> StructPropertyHandle,
 		TSharedRef<IPropertyHandle> ConnectionTargetsPropertyHandle,
diff --git a/Engine/Plugins/Experimental/VirtualCamera/Source/VCamCoreEditor/Private/Customizations/StateSwitcher/WidgetConnectionConfigTypeCustomization.h b/Engine/Plugins/Experimental/VirtualCamera/Source/VCamCoreEditor/Private/Customizations/StateSwitcher/WidgetConnectionConfigTypeCustomization.h
--- a/Engine/Plugins/Experimental/VirtualCamera/Source/VCamCoreEditor/Private/Customizations/StateSwitcher/WidgetConnectionConfigTypeCustomization.h
+++ b/Engine/Plugins/Experimental/VirtualCamera/Source/VCamCoreEditor/Private/Customizations/StateSwitcher/WidgetConnectionConfigTypeCustomization.h
@@ -25,10 +25,6 @@ namespace UE::VCamCoreEditor::Private
 
 	private:
 
-		void CustomizeWidgetReferenceProperty(TSharedRef<IPropertyHandle> WidgetReferencePropertyHandle, IDetailPropertyRow& Row, IPropertyTypeCustomizationUtils& CustomizationUtils) const;
-		TArray<FString> GetPropertyItemList(TWeakPtr<IPropertyUtilities> WeakPropertyUtils) const;
-		static TObjectPtr<UWidgetTree> GetWidgetTreeThroughBlueprintAsset(UUserWidget* ClassDefaultWidget); 
-
 		void CustomizeConnectionTargetsReferenceProperty(TSharedRef<IPropertyHandle> StructPropertyHandle, TSharedRef<IPropertyHandle> ConnectionTargetsPropertyHandle, IDetailChildrenBuilder& ChildBuilder, IPropertyTypeCustomizationUtils& CustomizationUtils) const;
 		TAttribute<TArray<FName>> CreateGetConnectionsFromChildWidgetAttribute(TSharedRef<IPropertyHandle> StructPropertyHandle, IPropertyTypeCustomizationUtils& CustomizationUtils) const;
 	};
