Fix VCamStateSwitcher always inits in "Default" state even though widget CDO defaults specify different state

diff --git a/Engine/Plugins/Experimental/VirtualCamera/Source/VCamCore/Private/UI/Switcher/VCamStateSwitcherWidget.cpp b/Engine/Plugins/Experimental/VirtualCamera/Source/VCamCore/Private/UI/Switcher/VCamStateSwitcherWidget.cpp
--- a/Engine/Plugins/Experimental/VirtualCamera/Source/VCamCore/Private/UI/Switcher/VCamStateSwitcherWidget.cpp
+++ b/Engine/Plugins/Experimental/VirtualCamera/Source/VCamCore/Private/UI/Switcher/VCamStateSwitcherWidget.cpp
@@ -8,7 +8,7 @@ DEFINE_LOG_CATEGORY(LogVCamStateSwitcher);
 
 FName UVCamStateSwitcherWidget::DefaultState("Default");
 
-bool UVCamStateSwitcherWidget::SetCurrentState(FName NewState, bool bForceUpdate)
+bool UVCamStateSwitcherWidget::SetCurrentState(FName NewState, bool bForceUpdate, bool bReinitializeConnections)
 {
 	if (NewState == CurrentState && !bForceUpdate)
 	{
@@ -36,7 +36,7 @@ bool UVCamStateSwitcherWidget::SetCurrentState(FName NewState, bool bForceUpdate
 		}
 
 		EConnectionUpdateResult UpdateResult;
-		Widget->UpdateConnectionTargets(WidgetConfig.ConnectionTargets, true, UpdateResult);
+		Widget->UpdateConnectionTargets(WidgetConfig.ConnectionTargets, bReinitializeConnections, UpdateResult);
 	}
 	
 	CurrentState = NewState;
@@ -77,10 +77,29 @@ void UVCamStateSwitcherWidget::NativePreConstruct()
 {
 	Super::NativePreConstruct();
 
+	// This widget was just created. There is not point in trying to update the current state because it will fail.
+	const bool bIsPendingInitialization = GetVCamComponent() == nullptr;
+	if (!bIsPendingInitialization)
+	{
+		SetStateOrFallbackToDefault(CurrentState);
+	}
+}
+
+void UVCamStateSwitcherWidget::OnInitializeConnections_Implementation(UVCamComponent* VCam)
+{
+	Super::OnInitializeConnections_Implementation(VCam);
+
+	// OnInitializeConnections is called just before ReinitializeConnections is called ... optimize by not calling it twice.
+	constexpr bool bReinitializeConnections = false;
+	SetStateOrFallbackToDefault(CurrentState, bReinitializeConnections);
+}
+
+void UVCamStateSwitcherWidget::SetStateOrFallbackToDefault(FName NewState, bool bReinitializeConnections)
+{
 	constexpr bool bForceUpdate = true;
-	const bool bIsCurrentStateValid = SetCurrentState(CurrentState, bForceUpdate); 
+	const bool bIsCurrentStateValid = SetCurrentState(NewState, bForceUpdate, bReinitializeConnections); 
 	if (!bIsCurrentStateValid)
 	{
 		SetCurrentState(DefaultState, bForceUpdate);
 	}
-}
+}
\ No newline at end of file
diff --git a/Engine/Plugins/Experimental/VirtualCamera/Source/VCamCore/Private/UI/VCamWidget.cpp b/Engine/Plugins/Experimental/VirtualCamera/Source/VCamCore/Private/UI/VCamWidget.cpp
--- a/Engine/Plugins/Experimental/VirtualCamera/Source/VCamCore/Private/UI/VCamWidget.cpp
+++ b/Engine/Plugins/Experimental/VirtualCamera/Source/VCamCore/Private/UI/VCamWidget.cpp
@@ -35,6 +35,7 @@ void UVCamWidget::InitializeConnections(UVCamComponent* VCam)
 		VCamComponent->AddInputMappingContext(InputMappingContext, InputContextPriority);
 	}
 
+	OnInitializeConnections(VCam);
 	ReinitializeConnections();
 }
 
@@ -107,3 +108,7 @@ void UVCamWidget::UpdateConnectionTargets(const TMap<FName, FVCamConnectionTarge
 		Result = EConnectionUpdateResult::NoConnectionsUpdated;
 	}
 }
+
+void UVCamWidget::OnInitializeConnections_Implementation(UVCamComponent* VCam)
+{
+}
diff --git a/Engine/Plugins/Experimental/VirtualCamera/Source/VCamCore/Public/UI/Switcher/VCamStateSwitcherWidget.h b/Engine/Plugins/Experimental/VirtualCamera/Source/VCamCore/Public/UI/Switcher/VCamStateSwitcherWidget.h
--- a/Engine/Plugins/Experimental/VirtualCamera/Source/VCamCore/Public/UI/Switcher/VCamStateSwitcherWidget.h
+++ b/Engine/Plugins/Experimental/VirtualCamera/Source/VCamCore/Public/UI/Switcher/VCamStateSwitcherWidget.h
@@ -32,9 +32,10 @@ public:
 	 * 
 	 * @param NewState The new state to switch to
 	 * @param bForceUpdate Call UpdateConnectionTargets even if the CurrentState == NewState
+	 * @param bReinitializeConnections Parameter to pass to UpdateConnectionTargets. If true, ReinitializeConnections will be called.
 	 */
 	UFUNCTION(BlueprintCallable, Category = "Connections")
-	bool SetCurrentState(FName NewState, bool bForceUpdate = false);
+	bool SetCurrentState(FName NewState, bool bForceUpdate = false, bool bReinitializeConnections = true);
 	
 	UFUNCTION(BlueprintPure, Category = "Connections")
 	FName GetCurrentState() const { return CurrentState; }
@@ -56,6 +57,10 @@ protected:
 	//~ Begin UUserWidget Interface
 	virtual void NativePreConstruct() override;
 	//~ End UUserWidget Interface
+
+	//~ Begin UVCamWidget Interface
+	virtual void OnInitializeConnections_Implementation(UVCamComponent* VCam) override;
+	//~ End UVCamWidget Interface
 	
 private:
 
@@ -75,4 +80,6 @@ private:
 
 	UPROPERTY(EditAnywhere, BlueprintGetter = "GetCurrentState", BlueprintSetter = "K2_SetCurrentState", Category = "Connections")
 	FName CurrentState = DefaultState;
+
+	void SetStateOrFallbackToDefault(FName NewState, bool bReinitializeConnections = true);
 };
diff --git a/Engine/Plugins/Experimental/VirtualCamera/Source/VCamCore/Public/UI/VCamWidget.h b/Engine/Plugins/Experimental/VirtualCamera/Source/VCamCore/Public/UI/VCamWidget.h
--- a/Engine/Plugins/Experimental/VirtualCamera/Source/VCamCore/Public/UI/VCamWidget.h
+++ b/Engine/Plugins/Experimental/VirtualCamera/Source/VCamCore/Public/UI/VCamWidget.h
@@ -29,10 +29,6 @@ class VCAMCORE_API UVCamWidget : public UUserWidget
 {
 	GENERATED_BODY()
 public:
-
-	UVCamWidget(const FObjectInitializer& ObjectInitializer);
-	
-	virtual void NativeDestruct() override;
 	
 	/*
 	 * The VCam Connections associated with this Widget
@@ -70,6 +66,12 @@ public:
 	UPROPERTY(BlueprintAssignable, meta = (DisplayName = "OnConnectionsReinitialized"))
 	FConnectionsReinitializedDelegate_Blueprint OnPostConnectionsReinitializedDelegate_Blueprint;
 	
+	UVCamWidget(const FObjectInitializer& ObjectInitializer);
+
+	//~ Begin UWidget Interface
+	virtual void NativeDestruct() override;
+	//~ End UWidget Interface
+	
 	/*
 	 * Event called when a specific connection has been updated
 	 * 
@@ -109,9 +111,18 @@ public:
 	UFUNCTION(BlueprintCallable, Category="VCam Connections", meta=(ExpandEnumAsExecs="Result", bReinitializeOnSuccessfulUpdate=true))
 	void UpdateConnectionTargets(const TMap<FName, FVCamConnectionTargetSettings>& NewConnectionTargets, const bool bReinitializeOnSuccessfulUpdate, EConnectionUpdateResult& Result);
 
+	FORCEINLINE TObjectPtr<UVCamComponent> GetVCamComponent() const { return VCamComponent; }
+
+protected:
+
+	UFUNCTION(BlueprintNativeEvent, Category = "VCam Connections")
+	void OnInitializeConnections(UVCamComponent* VCam);
+	
+private:
+	
 	/*
 	 * Cached pointer to the VCam Component that owns this VCam Widget
 	 */
-	UPROPERTY(Transient, BlueprintReadOnly, Category="VCam")
+	UPROPERTY(Transient, BlueprintReadOnly, Category = "VCam", meta = (AllowPrivateAccess = "true"))
 	TObjectPtr<UVCamComponent> VCamComponent;
 };
\ No newline at end of file
