Fix typo causing use of the wrong IPropertyHandle for quieries.
Fixes that TargetSettings do not respect RequiredInterface in drop-down and TargetSettings do not respect required input actions

diff --git a/Engine/Plugins/Experimental/VirtualCamera/Source/VCamCoreEditor/Private/Customizations/ConnectionTargetSettingsTypeCustomization.cpp b/Engine/Plugins/Experimental/VirtualCamera/Source/VCamCoreEditor/Private/Customizations/ConnectionTargetSettingsTypeCustomization.cpp
--- a/Engine/Plugins/Experimental/VirtualCamera/Source/VCamCoreEditor/Private/Customizations/ConnectionTargetSettingsTypeCustomization.cpp
+++ b/Engine/Plugins/Experimental/VirtualCamera/Source/VCamCoreEditor/Private/Customizations/ConnectionTargetSettingsTypeCustomization.cpp
@@ -72,9 +72,10 @@ namespace UE::VCamCoreEditor::Private
 		AddScopeRow(ChildBuilder, CustomizationUtils);
 		
 		const TSharedPtr<IPropertyHandle> ParentHandle = PropertyHandle->GetParentHandle();
+		const FStructProperty* ParentStruct = CastField<FStructProperty>(ParentHandle->GetProperty());
 		const bool bIsValidHandle = ParentHandle
-			&& CastField<FStructProperty>(PropertyHandle->GetProperty())
-			&& CastField<FStructProperty>(PropertyHandle->GetProperty())->Struct == FVCamConnection::StaticStruct();
+			&& ParentStruct
+			&& ParentStruct->Struct == FVCamConnection::StaticStruct();
 		const TSharedPtr<IPropertyHandle> OptionalVCamConnectionParentStructHandle = bIsValidHandle
 			? ParentHandle
 			: nullptr;
