For target settings customization, the drop-down should contain None and sort the remaining names alphabetically

diff --git a/Engine/Plugins/Experimental/VirtualCamera/Source/VCamCoreEditor/Private/Customizations/ConnectionTargetSettingsTypeCustomization.cpp b/Engine/Plugins/Experimental/VirtualCamera/Source/VCamCoreEditor/Private/Customizations/ConnectionTargetSettingsTypeCustomization.cpp
--- a/Engine/Plugins/Experimental/VirtualCamera/Source/VCamCoreEditor/Private/Customizations/ConnectionTargetSettingsTypeCustomization.cpp
+++ b/Engine/Plugins/Experimental/VirtualCamera/Source/VCamCoreEditor/Private/Customizations/ConnectionTargetSettingsTypeCustomization.cpp
@@ -4,16 +4,14 @@
 
 #include "VCamComponent.h"
 #include "UI/VCamConnectionStructs.h"
+#include "Util/ConnectionUtils.h"
 
 #include "Algo/ForEach.h"
 #include "DetailWidgetRow.h"
 #include "Framework/MultiBox/MultiBoxBuilder.h"
 #include "IDetailChildrenBuilder.h"
-#include "LevelEditor.h"
-#include "Modules/ModuleManager.h"
 #include "Selection.h"
 #include "SSimpleComboButton.h"
-#include "Util/ConnectionUtils.h"
 #include "Widgets/Images/SImage.h"
 #include "Widgets/Layout/SBox.h"
 #include "Widgets/SBoxPanel.h"
@@ -23,6 +21,14 @@
 
 namespace UE::VCamCoreEditor::Private
 {
+	static TArray<FName> SortNamesAlphabeticallyAndPrependNone(TArray<FName>&& Names)
+	{
+		TArray<FName> Result = MoveTemp(Names);
+		Result.Sort([](const FName& Left, const FName& Right){ return Left.LexicalLess(Right); });
+		Result.Insert(NAME_None, 0);
+		return Result;
+	}
+	
 	TSharedRef<IPropertyTypeCustomization> FConnectionTargetSettingsTypeCustomization::MakeInstance()
 	{
 		return MakeShared<FConnectionTargetSettingsTypeCustomization>();
@@ -144,10 +150,10 @@ namespace UE::VCamCoreEditor::Private
 					&& ValueData)
 				{
 					const FVCamConnection* ConnectionData = static_cast<FVCamConnection*>(ValueData);
-					return VCamCore::ConnectionUtils::FindCompatibleModifierNames(*ConnectionData, *DataSource);
+					return SortNamesAlphabeticallyAndPrependNone(VCamCore::ConnectionUtils::FindCompatibleModifierNames(*ConnectionData, *DataSource));
 				}
 				
-				return DataSource->GetAllModifierNames();
+				return SortNamesAlphabeticallyAndPrependNone(DataSource->GetAllModifierNames());
 			}),
 			TAttribute<bool>::CreateLambda([](){ return GetUserFocusedConnectionPointSource().ComponentSource != EComponentSource::None; })
 			);
@@ -189,12 +195,12 @@ namespace UE::VCamCoreEditor::Private
 					&& ValueData)
 				{
 					const FVCamConnection* ConnectionData = static_cast<FVCamConnection*>(ValueData);
-					return VCamCore::ConnectionUtils::FindCompatibleConnectionPoints(*ConnectionData, *Modifier);
+					return SortNamesAlphabeticallyAndPrependNone(VCamCore::ConnectionUtils::FindCompatibleConnectionPoints(*ConnectionData, *Modifier));
 				}
 
 				TArray<FName> ConnectionPoints;
 				Modifier->ConnectionPoints.GenerateKeyArray(ConnectionPoints);
-				return ConnectionPoints;
+				return SortNamesAlphabeticallyAndPrependNone(MoveTemp(ConnectionPoints));
 			}),
 			TAttribute<bool>::CreateLambda([ModifierHandle]()
 			{
