FChildWidgetReference loses reference when widget is recompiled

diff --git a/Engine/Plugins/Experimental/VirtualCamera/Source/VCamCore/Private/Util/WidgetReference.cpp b/Engine/Plugins/Experimental/VirtualCamera/Source/VCamCore/Private/Util/WidgetReference.cpp
--- a/Engine/Plugins/Experimental/VirtualCamera/Source/VCamCore/Private/Util/WidgetReference.cpp
+++ b/Engine/Plugins/Experimental/VirtualCamera/Source/VCamCore/Private/Util/WidgetReference.cpp
@@ -3,14 +3,13 @@
 #include "Util/WidgetReference.h"
 
 #include "LogVCamCore.h"
-#include "UI/Switcher/VCamStateSwitcherWidget.h"
 #include "Util/WidgetTreeUtils.h"
 
 #include "Blueprint/WidgetTree.h"
 
 UWidget* FChildWidgetReference::ResolveWidget(UUserWidget& OwnerWidget) const
 {
-	UWidget* ResolvedTemplate = Template.Get();
+	UWidget* ResolvedTemplate = Template.LoadSynchronous();
 	if (!ResolvedTemplate)
 	{
 		return nullptr;
diff --git a/Engine/Plugins/Experimental/VirtualCamera/Source/VCamCore/Public/Util/WidgetReference.h b/Engine/Plugins/Experimental/VirtualCamera/Source/VCamCore/Public/Util/WidgetReference.h
--- a/Engine/Plugins/Experimental/VirtualCamera/Source/VCamCore/Public/Util/WidgetReference.h
+++ b/Engine/Plugins/Experimental/VirtualCamera/Source/VCamCore/Public/Util/WidgetReference.h
@@ -24,7 +24,7 @@ struct VCAMCORE_API FChildWidgetReference
 	 * @see FWidgetReference
 	 */
 	UPROPERTY(EditAnywhere, Category = "Virtual Camera")
-	TLazyObjectPtr<UWidget> Template;
+	TSoftObjectPtr<UWidget> Template;
 	
 	UWidget* ResolveWidget(UUserWidget& OwnerWidget) const;
 	bool HasNoWidgetSet() const { return Template.IsValid(); }
