Forward proof hotfix by replacing inline functions with linked functions

diff --git a/Engine/Plugins/Experimental/VirtualCamera/Source/VCamCore/Private/Output/VCamOutputProviderBase.cpp b/Engine/Plugins/Experimental/VirtualCamera/Source/VCamCore/Private/Output/VCamOutputProviderBase.cpp
--- a/Engine/Plugins/Experimental/VirtualCamera/Source/VCamCore/Private/Output/VCamOutputProviderBase.cpp
+++ b/Engine/Plugins/Experimental/VirtualCamera/Source/VCamCore/Private/Output/VCamOutputProviderBase.cpp
@@ -152,6 +152,16 @@ void UVCamOutputProviderBase::SetTargetCamera(const UCineCameraComponent* InTarg
 	NotifyWidgetOfComponentChange();
 }
 
+void UVCamOutputProviderBase::SetTargetViewport(EVCamTargetViewportID Value)
+{
+	TargetViewport = Value;
+}
+
+void UVCamOutputProviderBase::SetUMGClass(const TSubclassOf<UUserWidget> InUMGClass)
+{
+	UMGClass = InUMGClass;
+}
+
 void UVCamOutputProviderBase::CreateUMG()
 {
 	if (!UMGClass)
diff --git a/Engine/Plugins/Experimental/VirtualCamera/Source/VCamCore/Public/Output/VCamOutputProviderBase.h b/Engine/Plugins/Experimental/VirtualCamera/Source/VCamCore/Public/Output/VCamOutputProviderBase.h
--- a/Engine/Plugins/Experimental/VirtualCamera/Source/VCamCore/Public/Output/VCamOutputProviderBase.h
+++ b/Engine/Plugins/Experimental/VirtualCamera/Source/VCamCore/Public/Output/VCamOutputProviderBase.h
@@ -82,12 +82,12 @@ public:
 	UFUNCTION(BlueprintPure, Category = "Output")
 	EVCamTargetViewportID GetTargetViewport() const { return TargetViewport; }
 	UFUNCTION(BlueprintCallable, Category = "Output")
-	void SetTargetViewport(EVCamTargetViewportID Value) { TargetViewport = Value; }
+	void SetTargetViewport(EVCamTargetViewportID Value);
 	
 	UFUNCTION(BlueprintPure, Category = "Output")
 	TSubclassOf<UUserWidget> GetUMGClass() const { return UMGClass; }
 	UFUNCTION(BlueprintCallable, Category = "Output")
-	void SetUMGClass(const TSubclassOf<UUserWidget> InUMGClass) { UMGClass = InUMGClass; }
+	void SetUMGClass(const TSubclassOf<UUserWidget> InUMGClass);
 
 	UVPFullScreenUserWidget* GetUMGWidget() { return UMGWidget; };
 
