Fix bad memory access when closing the editor.

diff --git a/Engine/Plugins/Experimental/VirtualCamera/Source/VCamCoreEditor/Private/VCamCoreEditorModule.cpp b/Engine/Plugins/Experimental/VirtualCamera/Source/VCamCoreEditor/Private/VCamCoreEditorModule.cpp
--- a/Engine/Plugins/Experimental/VirtualCamera/Source/VCamCoreEditor/Private/VCamCoreEditorModule.cpp
+++ b/Engine/Plugins/Experimental/VirtualCamera/Source/VCamCoreEditor/Private/VCamCoreEditorModule.cpp
@@ -55,7 +55,10 @@ namespace UE::VCamCoreEditor::Private
 	void FVCamCoreEditorModule::ShutdownModule()
 	{
 		FKismetEditorUtilities::UnregisterAutoBlueprintNodeCreation(this);
-		UnregisterCustomizations();
+		if (UObjectInitialized() && !IsEngineExitRequested())
+		{
+			UnregisterCustomizations();
+		}
 	}
 
 	void FVCamCoreEditorModule::RegisterConnectionRemapCustomization(TSubclassOf<UVCamWidget> Class, FGetConnectionRemappingCustomization GetterDelegate)
@@ -170,19 +173,20 @@ namespace UE::VCamCoreEditor::Private
 
 	void FVCamCoreEditorModule::UnregisterCustomizations()
 	{
-		FPropertyEditorModule& PropertyModule = FModuleManager::GetModuleChecked<FPropertyEditorModule>("PropertyEditor");
-		
-		PropertyModule.UnregisterCustomPropertyTypeLayout(FVCamInputProfile::StaticStruct()->GetFName());
-		PropertyModule.UnregisterCustomPropertyTypeLayout(FVCamConnectionTargetSettings::StaticStruct()->GetFName());
-		PropertyModule.UnregisterCustomPropertyTypeLayout(FWidgetConnectionConfig::StaticStruct()->GetFName());
-		PropertyModule.UnregisterCustomPropertyTypeLayout(UVCamOutputProviderBase::StaticClass()->GetFName());
-		PropertyModule.UnregisterCustomPropertyTypeLayout(FVCamViewportLocker::StaticStruct()->GetFName());
-		PropertyModule.UnregisterCustomPropertyTypeLayout(FChildWidgetReference::StaticStruct()->GetFName());
-		PropertyModule.UnregisterCustomPropertyTypeLayout(FVCamChildWidgetReference::StaticStruct()->GetFName());
-		
-		PropertyModule.UnregisterCustomClassLayout(UVCamOutputProviderBase::StaticClass()->GetFName());
-		PropertyModule.UnregisterCustomClassLayout(AVCamBaseActor::StaticClass()->GetFName());
-		PropertyModule.UnregisterCustomClassLayout(UVCamStateSwitcherWidget::StaticClass()->GetFName());
+		if (FPropertyEditorModule* PropertyModule = FModuleManager::GetModulePtr<FPropertyEditorModule>("PropertyEditor"))
+		{
+			PropertyModule->UnregisterCustomPropertyTypeLayout(FVCamInputProfile::StaticStruct()->GetFName());
+			PropertyModule->UnregisterCustomPropertyTypeLayout(FVCamConnectionTargetSettings::StaticStruct()->GetFName());
+			PropertyModule->UnregisterCustomPropertyTypeLayout(FWidgetConnectionConfig::StaticStruct()->GetFName());
+			PropertyModule->UnregisterCustomPropertyTypeLayout(UVCamOutputProviderBase::StaticClass()->GetFName());
+			PropertyModule->UnregisterCustomPropertyTypeLayout(FVCamViewportLocker::StaticStruct()->GetFName());
+			PropertyModule->UnregisterCustomPropertyTypeLayout(FChildWidgetReference::StaticStruct()->GetFName());
+			PropertyModule->UnregisterCustomPropertyTypeLayout(FVCamChildWidgetReference::StaticStruct()->GetFName());
+
+			PropertyModule->UnregisterCustomClassLayout(UVCamOutputProviderBase::StaticClass()->GetFName());
+			PropertyModule->UnregisterCustomClassLayout(AVCamBaseActor::StaticClass()->GetFName());
+			PropertyModule->UnregisterCustomClassLayout(UVCamStateSwitcherWidget::StaticClass()->GetFName());
+		}	
 	}
 }
 
