After editing the VCamActor BP, the Widget is not added to the Viewport

diff --git a/Engine/Plugins/Experimental/VirtualCamera/Source/VCamCore/Private/Output/VCamOutputProviderBase.cpp b/Engine/Plugins/Experimental/VirtualCamera/Source/VCamCore/Private/Output/VCamOutputProviderBase.cpp
--- a/Engine/Plugins/Experimental/VirtualCamera/Source/VCamCore/Private/Output/VCamOutputProviderBase.cpp
+++ b/Engine/Plugins/Experimental/VirtualCamera/Source/VCamCore/Private/Output/VCamOutputProviderBase.cpp
@@ -148,8 +148,11 @@ bool UVCamOutputProviderBase::IsOuterComponentEnabled() const
 
 void UVCamOutputProviderBase::SetTargetCamera(const UCineCameraComponent* InTargetCamera)
 {
-	TargetCamera = InTargetCamera;
-	NotifyWidgetOfComponentChange();
+	if (InTargetCamera != TargetCamera)
+	{
+		TargetCamera = InTargetCamera;
+		NotifyWidgetOfComponentChange();
+	}
 }
 
 void UVCamOutputProviderBase::SetTargetViewport(EVCamTargetViewportID Value)
@@ -244,6 +247,7 @@ void UVCamOutputProviderBase::DisplayUMG()
 
 		if (ActorWorld)
 		{
+			UMGWidget->SetCustomPostProcessSettingsSource(TargetCamera.Get());
 			UMGWidget->Display(ActorWorld);
 			if (UUserWidget* Subwidget = UMGWidget->GetWidget(); ensure(Subwidget) && WidgetSnapshot.HasData())
 			{
@@ -319,8 +323,6 @@ void UVCamOutputProviderBase::NotifyWidgetOfComponentChange() const
 {
 	if (UMGWidget && UMGWidget->IsDisplayed())
 	{
-		UMGWidget->SetCustomPostProcessSettingsSource(TargetCamera.Get());
-		
 		UUserWidget* DisplayedWidget = UMGWidget->GetWidget();
 		if (IsValid(DisplayedWidget))
 		{
@@ -525,7 +527,10 @@ void UVCamOutputProviderBase::PostEditChangeProperty(FPropertyChangedEvent& Prop
 			WidgetSnapshot.Reset();
 			if (bIsActive)
 			{
+				// In case a child class resets UMGClass, reapply the correct value we got the PostEditChangeProperty for.
+				const TSubclassOf<UUserWidget> ProtectUMGClass = UMGClass;
 				SetActive(false);
+				UMGClass = ProtectUMGClass;
 				SetActive(true);
 			}
 		}
@@ -535,7 +540,7 @@ void UVCamOutputProviderBase::PostEditChangeProperty(FPropertyChangedEvent& Prop
 			{
 				RestoreOverrideResolutionForViewport(static_cast<EVCamTargetViewportID>(i));
 			}
-			ApplyOverrideResolutionForViewport(TargetViewport);
+			ReapplyOverrideResolution(TargetViewport);
 
 			if (bIsActive)
 			{
diff --git a/Engine/Plugins/Experimental/VirtualCamera/Source/VCamCore/Private/Util/LevelViewportUtils.cpp b/Engine/Plugins/Experimental/VirtualCamera/Source/VCamCore/Private/Util/LevelViewportUtils.cpp
--- a/Engine/Plugins/Experimental/VirtualCamera/Source/VCamCore/Private/Util/LevelViewportUtils.cpp
+++ b/Engine/Plugins/Experimental/VirtualCamera/Source/VCamCore/Private/Util/LevelViewportUtils.cpp
@@ -29,19 +29,19 @@ namespace UE::VCamCore::LevelViewportUtils::Private
 		ENUM_CLASS_FLAGS(EViewportFlags)
 		
 #if WITH_EDITOR
-		static void UpdateLockStateForEditor(FVCamViewportLockState& ViewportLockState, EVCamTargetViewportID ViewportID, bool bNewLockState, AActor* ActorToLockWith = nullptr)
+		static void UpdateLockStateForEditor(FVCamViewportLockState& ViewportLockState, EVCamTargetViewportID ViewportID, bool bNewLockState, AActor& ActorToLockWith)
 		{
 			TSharedPtr<SLevelViewport> Viewport = GetLevelViewport(ViewportID);
 			FLevelEditorViewportClient* LevelViewportClient = Viewport.IsValid()
 				? &Viewport->GetLevelViewportClient()
 				: nullptr;
-			const bool bNeedsLock = !ViewportLockState.bIsLockedToViewport && bNewLockState;
 			if (LevelViewportClient)
 			{
+				const bool bNeedsLock = !ViewportLockState.bIsLockedToViewport && bNewLockState;
 				if (bNeedsLock)
 				{
 					ViewportLockState.Backup_ActorLock = LevelViewportClient->GetActiveActorLock();
-					LevelViewportClient->SetActorLock(ActorToLockWith);
+					LevelViewportClient->SetActorLock(&ActorToLockWith);
 					// If bLockedCameraView is not true then the viewport is locked to the actor's transform and not the camera component
 					LevelViewportClient->bLockedCameraView = true;
 					ViewportLockState.bIsLockedToViewport = true;
@@ -64,7 +64,7 @@ namespace UE::VCamCore::LevelViewportUtils::Private
 		}
 #endif
 
-		static void UpdateLockStateForGame(FVCamViewportLockState& ViewportLockState, const FWorldContext& Context, bool bNewLockState, AActor* ActorToLockWith = nullptr)
+		static void UpdateLockStateForGame(FVCamViewportLockState& ViewportLockState, const FWorldContext& Context, bool bNewLockState, AActor& ActorToLockWith)
 		{
 			UWorld* ActorWorld = Context.World();
 			if (ActorWorld && ActorWorld->GetGameInstance())
@@ -76,7 +76,7 @@ namespace UE::VCamCore::LevelViewportUtils::Private
 					if (bNeedsLock)
 					{
 						ViewportLockState.Backup_ViewTarget = PlayerController->GetViewTarget();
-						PlayerController->SetViewTarget(ActorToLockWith);
+						PlayerController->SetViewTarget(&ActorToLockWith);
 						ViewportLockState.bIsLockedToViewport = true;
 					}
 					else if (ViewportLockState.bIsLockedToViewport && !bNewLockState
@@ -95,7 +95,7 @@ namespace UE::VCamCore::LevelViewportUtils::Private
 			}
 		}
 		
-		static void UpdateLockState(FVCamViewportLockState& ViewportLockState, EVCamTargetViewportID ViewportID, EViewportFlags ViewportFlags, AActor* ActorToLockWith = nullptr)
+		static void UpdateLockState(FVCamViewportLockState& ViewportLockState, EVCamTargetViewportID ViewportID, EViewportFlags ViewportFlags, AActor& ActorToLockWith)
 		{
 			const bool bForceUse = (ViewportFlags & EViewportFlags::ForceUse) != EViewportFlags::None;
 			const bool bIsUsed = (ViewportFlags & EViewportFlags::IsNotUsed) == EViewportFlags::None;
@@ -120,7 +120,7 @@ namespace UE::VCamCore::LevelViewportUtils::Private
 		}
 	}
 
-	void UpdateViewportLocksFromOutputs(TArray<TObjectPtr<UVCamOutputProviderBase>> OutputProviders, FVCamViewportLocker& LockData, AActor* ActorToLockWith)
+	void UpdateViewportLocksFromOutputs(TArray<TObjectPtr<UVCamOutputProviderBase>> OutputProviders, FVCamViewportLocker& LockData, AActor& ActorToLockWith)
 	{
 		TSet<EVCamTargetViewportID> Viewports;
 		TSet<EVCamTargetViewportID> ForcefullyLockedViewports;
@@ -148,11 +148,11 @@ namespace UE::VCamCore::LevelViewportUtils::Private
 		}
 	}
 
-	void UnlockAllViewports(FVCamViewportLocker& LockData)
+	void UnlockAllViewports(FVCamViewportLocker& LockData, AActor& ActorToLockWith)
 	{
 		for (TPair<EVCamTargetViewportID, FVCamViewportLockState>& ViewportData : LockData.Locks)
 		{
-			Locking::UpdateLockState(ViewportData.Value, ViewportData.Key, Locking::EViewportFlags::IsNotUsed);
+			Locking::UpdateLockState(ViewportData.Value, ViewportData.Key, Locking::EViewportFlags::IsNotUsed, ActorToLockWith);
 		}
 	}
 
@@ -161,14 +161,6 @@ namespace UE::VCamCore::LevelViewportUtils::Private
 	{
 		TSharedPtr<SLevelViewport> OutLevelViewport = nullptr;
 
-		if (TargetViewport == EVCamTargetViewportID::CurrentlySelected)
-		{
-			FLevelEditorModule* LevelEditorModule = FModuleManager::GetModulePtr<FLevelEditorModule>(TEXT("LevelEditor"));
-			return LevelEditorModule
-				? LevelEditorModule->GetFirstActiveLevelViewport()
-				: nullptr;;
-		}
-
 		if (!GEditor)
 		{
 			return nullptr;
@@ -201,9 +193,7 @@ namespace UE::VCamCore::LevelViewportUtils::Private
 	
 	FString GetConfigKeyFor(EVCamTargetViewportID TargetViewport)
 	{
-		return TargetViewport == EVCamTargetViewportID::CurrentlySelected
-			? FName(EName::None).ToString()
-			: FString::Printf(TEXT("Viewport %d.Viewport"), (int32)TargetViewport);
+		return FString::Printf(TEXT("Viewport %d.Viewport"), static_cast<int32>(TargetViewport) + 1);
 	}
 #endif
 }
diff --git a/Engine/Plugins/Experimental/VirtualCamera/Source/VCamCore/Private/Util/LevelViewportUtils.h b/Engine/Plugins/Experimental/VirtualCamera/Source/VCamCore/Private/Util/LevelViewportUtils.h
--- a/Engine/Plugins/Experimental/VirtualCamera/Source/VCamCore/Private/Util/LevelViewportUtils.h
+++ b/Engine/Plugins/Experimental/VirtualCamera/Source/VCamCore/Private/Util/LevelViewportUtils.h
@@ -13,9 +13,9 @@ struct FVCamViewportLocker;
 namespace UE::VCamCore::LevelViewportUtils::Private
 {
 	/** Locks or unlocks all viewports which are used by OutputProviders depending on the corresponding FVCamViewportLockState::bLockViewportToCamera flag. All other viewports will be unlocked. */
-	void UpdateViewportLocksFromOutputs(TArray<TObjectPtr<UVCamOutputProviderBase>> OutputProviders, FVCamViewportLocker& LockData, AActor* ActorToLockWith);
+	void UpdateViewportLocksFromOutputs(TArray<TObjectPtr<UVCamOutputProviderBase>> OutputProviders, FVCamViewportLocker& LockData, AActor& ActorToLockWith);
 	/** Sets all viewports to be unlocked. */
-	void UnlockAllViewports(FVCamViewportLocker& LockData);
+	void UnlockAllViewports(FVCamViewportLocker& LockData, AActor& ActorToLockWith);
 
 #if WITH_EDITOR
 	/** Gets the level viewport identified by TargetViewport */
diff --git a/Engine/Plugins/Experimental/VirtualCamera/Source/VCamCore/Private/VCamComponent.cpp b/Engine/Plugins/Experimental/VirtualCamera/Source/VCamCore/Private/VCamComponent.cpp
--- a/Engine/Plugins/Experimental/VirtualCamera/Source/VCamCore/Private/VCamComponent.cpp
+++ b/Engine/Plugins/Experimental/VirtualCamera/Source/VCamCore/Private/VCamComponent.cpp
@@ -100,7 +100,7 @@ UVCamComponent::UVCamComponent()
 void UVCamComponent::OnComponentDestroyed(bool bDestroyingHierarchy)
 {
 	bLockViewportToCamera_DEPRECATED = false;
-	UpdateActorViewportLocks();
+	UnlockAllViewports();
 
 	for (UVCamOutputProviderBase* Provider : OutputProviders)
 	{
@@ -158,19 +158,26 @@ void UVCamComponent::HandleObjectReplaced(const TMap<UObject*, UObject*>& Replac
 
 void UVCamComponent::NotifyComponentWasReplaced(UVCamComponent* ReplacementComponent)
 {
-	// This function should only ever be called when we have a valid component replacing us
 	check(ReplacementComponent);
 
 	// Make sure to copy over our delegate bindings to the component replacing us
 	ReplacementComponent->OnComponentReplaced = OnComponentReplaced;
-
 	OnComponentReplaced.Clear();
+	
+	const bool bWasEnabled = bEnabled;
 	// Ensure all modifiers and output providers get deinitialized
-	SetEnabled(false);
-
-	// Refresh the enabled state on the new component to prevent any stale state from the replacement
-	if (ReplacementComponent->IsEnabled())
+	if (bEnabled)
+	{
+		SetEnabled(false);
+	}
+	
+	// Refresh the enabled state on the new component to prevent any stale state from the replacement but only do so once.
+	// Careful: NotifyComponentWasReplaced can be called multiple times with the same arguments.
+	// If this is not the first time NotifyComponentWasReplaced is called, then bWasEnabled will already be false (see above).
+	const bool bNeedsToStartupOutputProviders = bWasEnabled && ReplacementComponent->IsEnabled();
+	if (bNeedsToStartupOutputProviders)
 	{
+		ReplacementComponent->ViewportLocker.Reset();
 		ReplacementComponent->SetEnabled(false);
 		ReplacementComponent->SetEnabled(true);
 	}
@@ -476,6 +483,7 @@ void UVCamComponent::OnOutputProvidersEdited(FPropertyChangedChainEvent& Propert
 			// Initialized inside the Update() loop
 			if (ChangedProvider && ShouldUpdateOutputProviders())
 			{
+				ChangedProvider->SetTargetCamera(GetTargetCamera());
 				ChangedProvider->Initialize();
 			}
 		}
@@ -1259,7 +1267,7 @@ void UVCamComponent::UpdateActorViewportLocks()
 	if (const UCineCameraComponent* Camera = GetTargetCamera()
 		; Camera && ensure(Camera->GetOwner()))
 	{
-		UE::VCamCore::LevelViewportUtils::Private::UpdateViewportLocksFromOutputs(OutputProviders, ViewportLocker, Camera->GetOwner());
+		UE::VCamCore::LevelViewportUtils::Private::UpdateViewportLocksFromOutputs(OutputProviders, ViewportLocker, *Camera->GetOwner());
 	}
 }
 
@@ -1268,7 +1276,7 @@ void UVCamComponent::UnlockAllViewports()
 	if (const UCineCameraComponent* Camera = GetTargetCamera()
 		; Camera && ensure(Camera->GetOwner()))
 	{
-		UE::VCamCore::LevelViewportUtils::Private::UnlockAllViewports(ViewportLocker);
+		UE::VCamCore::LevelViewportUtils::Private::UnlockAllViewports(ViewportLocker, *Camera->GetOwner());
 	}
 }
 
diff --git a/Engine/Plugins/Experimental/VirtualCamera/Source/VCamCore/Public/EVCamTargetViewportID.h b/Engine/Plugins/Experimental/VirtualCamera/Source/VCamCore/Public/EVCamTargetViewportID.h
--- a/Engine/Plugins/Experimental/VirtualCamera/Source/VCamCore/Public/EVCamTargetViewportID.h
+++ b/Engine/Plugins/Experimental/VirtualCamera/Source/VCamCore/Public/EVCamTargetViewportID.h
@@ -8,11 +8,10 @@
 UENUM(BlueprintType, meta=(DisplayName = "VCam Target Viewport ID"))
 enum class EVCamTargetViewportID : uint8
 {
-	CurrentlySelected = 0,
-	Viewport1 = 1,
-	Viewport2 = 2,
-	Viewport3 = 3,
-	Viewport4 = 4,
+	Viewport1 = 0,
+	Viewport2 = 1,
+	Viewport3 = 2,
+	Viewport4 = 3,
 
 	Count UMETA(Hidden)
 };
\ No newline at end of file
diff --git a/Engine/Plugins/Experimental/VirtualCamera/Source/VCamCore/Public/Output/VCamOutputProviderBase.h b/Engine/Plugins/Experimental/VirtualCamera/Source/VCamCore/Public/Output/VCamOutputProviderBase.h
--- a/Engine/Plugins/Experimental/VirtualCamera/Source/VCamCore/Public/Output/VCamOutputProviderBase.h
+++ b/Engine/Plugins/Experimental/VirtualCamera/Source/VCamCore/Public/Output/VCamOutputProviderBase.h
@@ -160,7 +160,7 @@ private:
 
 	/** Which viewport to use for this VCam */
 	UPROPERTY(EditAnywhere, BlueprintGetter = "GetTargetViewport", BlueprintSetter = "SetTargetViewport", Category = "Output", meta = (DisplayPriority = "2"))
-	EVCamTargetViewportID TargetViewport = EVCamTargetViewportID::CurrentlySelected;
+	EVCamTargetViewportID TargetViewport = EVCamTargetViewportID::Viewport1;
 	
 	/** The UMG class to be rendered in this output provider */
 	UPROPERTY(EditAnywhere, BlueprintGetter = "GetUMGClass", BlueprintSetter = "SetUMGClass", Category = "Output", meta = (DisplayName="UMG Overlay", DisplayPriority = "3"))
diff --git a/Engine/Plugins/Experimental/VirtualCamera/Source/VCamCore/Public/Util/VCamViewportLocker.h b/Engine/Plugins/Experimental/VirtualCamera/Source/VCamCore/Public/Util/VCamViewportLocker.h
--- a/Engine/Plugins/Experimental/VirtualCamera/Source/VCamCore/Public/Util/VCamViewportLocker.h
+++ b/Engine/Plugins/Experimental/VirtualCamera/Source/VCamCore/Public/Util/VCamViewportLocker.h
@@ -49,4 +49,14 @@ struct FVCamViewportLocker
 		{ EVCamTargetViewportID::Viewport3, {} },
 		{ EVCamTargetViewportID::Viewport4, {} }
 	};
+
+	void Reset()
+	{
+		for (TPair<EVCamTargetViewportID, FVCamViewportLockState>& Pair : Locks)
+		{
+			Pair.Value.bIsLockedToViewport = false;
+			Pair.Value.Backup_ActorLock = nullptr;
+			Pair.Value.Backup_ViewTarget = nullptr;
+		}
+	}
 };
diff --git a/Engine/Plugins/Experimental/VirtualCamera/Source/VCamCore/Public/VCamComponent.h b/Engine/Plugins/Experimental/VirtualCamera/Source/VCamCore/Public/VCamComponent.h
--- a/Engine/Plugins/Experimental/VirtualCamera/Source/VCamCore/Public/VCamComponent.h
+++ b/Engine/Plugins/Experimental/VirtualCamera/Source/VCamCore/Public/VCamComponent.h
@@ -372,7 +372,7 @@ public:
 	 * See FVCamCoreCustomVersion::MoveTargetViewportFromComponentToOutput.
 	 */
 	UPROPERTY()
-	EVCamTargetViewportID TargetViewport_DEPRECATED = EVCamTargetViewportID::CurrentlySelected;
+	EVCamTargetViewportID TargetViewport_DEPRECATED = EVCamTargetViewportID::Viewport1;
 
 private:
 	
diff --git a/Engine/Plugins/Experimental/VirtualProductionUtilities/Source/VPUtilities/Private/VPFullScreenUserWidget.cpp b/Engine/Plugins/Experimental/VirtualProductionUtilities/Source/VPUtilities/Private/VPFullScreenUserWidget.cpp
--- a/Engine/Plugins/Experimental/VirtualProductionUtilities/Source/VPUtilities/Private/VPFullScreenUserWidget.cpp
+++ b/Engine/Plugins/Experimental/VirtualProductionUtilities/Source/VPUtilities/Private/VPFullScreenUserWidget.cpp
@@ -378,6 +378,10 @@ bool FVPFullScreenUserWidget_PostProcess::InitPostProcessComponent(UWorld* World
 		}
 		
 		PostProcessMaterialInstance = UMaterialInstanceDynamic::Create(PostProcessMaterial, World);
+		if (!ensure(PostProcessMaterialInstance))
+		{
+			return false;
+		}
 
 		// set the parameter immediately
 		PostProcessMaterialInstance->SetTextureParameterValue(NAME_SlateUI, WidgetRenderTarget);
@@ -389,10 +393,11 @@ bool FVPFullScreenUserWidget_PostProcess::InitPostProcessComponent(UWorld* World
 			PostProcessSettings->WeightedBlendables.Array.SetNumZeroed(1);
 			PostProcessSettings->WeightedBlendables.Array[0].Weight = 1.f;
 			PostProcessSettings->WeightedBlendables.Array[0].Object = PostProcessMaterialInstance;
+			return true;
 		}
 	}
 
-	return PostProcessComponent && PostProcessMaterialInstance;
+	return false;
 }
 
 void FVPFullScreenUserWidget_PostProcess::ReleasePostProcessComponent()
