Fix output provider array elements not having any reset button

diff --git a/Engine/Plugins/Experimental/VirtualCamera/Source/VCamCoreEditor/Private/Customizations/OutputProvider/OutputProviderTypeCustomization.cpp b/Engine/Plugins/Experimental/VirtualCamera/Source/VCamCoreEditor/Private/Customizations/OutputProvider/OutputProviderTypeCustomization.cpp
--- a/Engine/Plugins/Experimental/VirtualCamera/Source/VCamCoreEditor/Private/Customizations/OutputProvider/OutputProviderTypeCustomization.cpp
+++ b/Engine/Plugins/Experimental/VirtualCamera/Source/VCamCoreEditor/Private/Customizations/OutputProvider/OutputProviderTypeCustomization.cpp
@@ -29,9 +29,13 @@ namespace UE::VCamCoreEditor::Private
 		UVCamOutputProviderBase* CustomizedOutputProvider = bSuccessGettingValue
 			? Cast<UVCamOutputProviderBase>(CustomizedObject)
 			: nullptr;
-		if (!CustomizedOutputProvider)
+		
+		uint32 NumChildren;
+		// If user just clicked + Add button or reset the instance class, reference will be null. Let the user select a class.
+		if (!CustomizedOutputProvider
+			|| PropertyHandle->GetNumChildren(NumChildren) != FPropertyAccess::Success
+			|| NumChildren != 1)
 		{
-			// If user just clicked + Add button, reference will be null. Let the user select a class.
 			ChildBuilder.AddCustomRow(FText::GetEmpty())
 				.NameContent()
 				[
@@ -41,32 +45,44 @@ namespace UE::VCamCoreEditor::Private
 				[
 					PropertyHandle->CreatePropertyValueWidget()
 				];
-			return;
+
+			// Sadly this will be missing the reset to default widget - I could not figure out any way quickly to add it using AddCustomRow
 		}
-		
-		uint32 NumChildren;
-		if (PropertyHandle->GetNumChildren(NumChildren) != FPropertyAccess::Success
-			|| NumChildren != 1)
+		else
 		{
-			return;
-		}
+			// AddExternalObjects causes FOutputProviderLayoutCustomization to customize the details.
+			IDetailPropertyRow* DetailRow = ChildBuilder.AddExternalObjects({ CustomizedOutputProvider },
+				FAddPropertyParams()
+					.CreateCategoryNodes(false)		// Avoid creating intermediate group expansion
+					.AllowChildren(true)			// Child properties should be shown
+					.HideRootObjectNode(false)		// Needed so we can use NameContent() & ValueContent() below
+					);
+			
+			DetailRow->CustomWidget(true)
+				.NameContent()
+				[
+					PropertyHandle->CreatePropertyNameWidget()
+				]
+				.ValueContent()
+				[
+					PropertyHandle->CreatePropertyValueWidget()
+				];
 
-		// AddExternalObjects causes FOutputProviderLayoutCustomization to customize the details.
-		IDetailPropertyRow* DetailRow = ChildBuilder.AddExternalObjects({ CustomizedOutputProvider },
-			FAddPropertyParams()
-				.CreateCategoryNodes(false)		// Avoid creating intermediate group expansion
-				.AllowChildren(true)			// Child properties should be shown
-				.HideRootObjectNode(false)		// Needed so we can use NameContent() & ValueContent() below
-				);
-		DetailRow->CustomWidget(true)
-			.NameContent()
-			[
-				PropertyHandle->CreatePropertyNameWidget()
-			]
-			.ValueContent()
-			[
-				PropertyHandle->CreatePropertyValueWidget()
-			];
+			// Because AddExternalObjects was used the property system will not add a reset to default widget by default 
+			DetailRow->OverrideResetToDefault(
+				FResetToDefaultOverride::Create(
+					FIsResetToDefaultVisible::CreateLambda([PropertyHandle](TSharedPtr<IPropertyHandle> Handle)
+					{
+						return PropertyHandle->CanResetToDefault();
+					}),
+					FResetToDefaultHandler::CreateLambda([PropertyHandle](TSharedPtr<IPropertyHandle> Handle)
+					{
+						return PropertyHandle->ResetToDefault();
+					})
+				)
+			);
+		}
+		
 	}
 }
 
