Fix input settings not being editable

diff --git a/Engine/Plugins/Experimental/VirtualCamera/Config/DefaultVCamInputSettings.ini b/Engine/Plugins/Experimental/VirtualCamera/Config/DefaultVCamInputSettings.ini
new file mode 100644
--- /dev/null
+++ b/Engine/Plugins/Experimental/VirtualCamera/Config/DefaultVCamInputSettings.ini
@@ -0,0 +1,3 @@
+[/Script/VCamInput.VCamInputSettings]
+DefaultInputProfile="VCam Actor Defaults"
+InputProfiles=(("VCam Actor Defaults", (MappableKeyOverrides=(("Filmback", None),("NextFilmbackPreset", None),("PreviousFilmbackPreset", None),("ClipPlaneAxis", None),("NextReticle", None),("PreviousReticle", None),("ReticleOpacityAxis", None),("ReticleOpacity", None),("Reticle", None),("NextOverlay", None),("PreviousOverlay", None),("OverlayOpacityAxis", None),("OverlayOpacity", None),("Overlay", None),("CustomAspectRatio", None),("CustomAspectRatioPrecision", None),("ToggleKillRoll", None),("SetCurrentLocationAxisScaling", None),("SetAllLocationAxisScaling", None),("SelectCurrentLocationAxisScaling", None),("SetLocationStabilization", None),("SetRotationStabilization", None),("ToggleHold", None),("Crane Movement", Gamepad_RightY),("Planar Movement", Gamepad_Left2D),("SetJoystickMovementGain", None),("Pan", Gamepad_RightX),("SetJoystickRotationGain", None),("ToggleFlightMode", None),("FocalLength", None),("NextLens", None),("PreviousLens", None),("Lens", None),("Iris", None),("StopDown", None),("StopUp", None),("FocusDistance", None),("FocusDistanceAxis", None),("FocalLengthAxis", None),("ISO", None),("NextISO", None),("PreviousISO", None),("ExposureCompensation", None),("ExposureCompensationReset", None),("CreateBookmark", None),("DeleteCurrentBookmark", None),("NextBookmark", None),("PreviousBookmark", None),("CurrentBookmark", None),("NextFrame", None),("PreviousFrame", None),("SetScrubSpeed", None),("SetSequenceFrameNormalized", None),("TogglePlayPause", None),("ToPlaybackEnd", None),("ToPlaybackStart", None),("StartRecording", None),("StopRecording", None),("ToggleRecording", None),("ZeroToStage", None),("ToggleSensorCorrection", None)))))
\ No newline at end of file
diff --git a/Engine/Plugins/Experimental/VirtualCamera/Source/VCamInput/Private/VCamInputSettings.cpp b/Engine/Plugins/Experimental/VirtualCamera/Source/VCamInput/Private/VCamInputSettings.cpp
--- a/Engine/Plugins/Experimental/VirtualCamera/Source/VCamInput/Private/VCamInputSettings.cpp
+++ b/Engine/Plugins/Experimental/VirtualCamera/Source/VCamInput/Private/VCamInputSettings.cpp
@@ -2,6 +2,9 @@
 
 #include "VCamInputSettings.h"
 
+#include "Interfaces/IPluginManager.h"
+#include "Misc/ConfigCacheIni.h"
+
 bool FVCamInputProfile::operator==(const FVCamInputProfile& OtherProfile) const
 {
 	bool bIsEqual = MappableKeyOverrides.Num() == OtherProfile.MappableKeyOverrides.Num();
@@ -42,6 +45,28 @@ TArray<FName> UVCamInputSettings::GetInputProfileNames() const
 	return ProfileNames;
 }
 
+void UVCamInputSettings::PostInitProperties()
+{
+	Super::PostInitProperties();
+
+	// Our flow is a bit unusual.
+	// UVCamInputSettings uses ProjectName/Saved/Platform/Game.ini to save & load settings.
+	// We've got a bunch of default settings in DefaultVCamInputSettings.ini.
+	// This allows us to iterate faster than setting up the default values in the C++ constructor.
+	// If we want default values to change, artists can just copy the settings from Game.ini into DefaultVCamInputSettings.ini..
+	// On the first run, we load from DefaultVCamInputSettings.ini.
+	// After that the engine will continue to use Game.ini.
+	const TSharedPtr<IPlugin> VirtualCameraPlugin = IPluginManager::Get().FindPlugin(TEXT("VirtualCamera"));
+	const FString PluginDefaultConfigFile = VirtualCameraPlugin->GetBaseDir() / TEXT("Config") / TEXT("DefaultVCamInputSettings.ini");
+	const FString RealDefaultConfigFile = GetConfigFilename(this);
+	const bool bNeedsDefaultValues = !GConfig->DoesSectionExist(TEXT("/Script/VCamInput.VCamInputSettings"), RealDefaultConfigFile);
+	const bool bConfigStillExists = GConfig->DoesSectionExist(TEXT("/Script/VCamInput.VCamInputSettings"), PluginDefaultConfigFile);
+	if (bNeedsDefaultValues && ensureMsgf(bConfigStillExists, TEXT("The content or location of DefaultVCamInputSettings.ini has changed")))
+	{
+		LoadConfig(StaticClass(), *PluginDefaultConfigFile);
+	}
+}
+
 UVCamInputSettings* UVCamInputSettings::GetVCamInputSettings()
 {
 	return GetMutableDefault<UVCamInputSettings>();
diff --git a/Engine/Plugins/Experimental/VirtualCamera/Source/VCamInput/Public/VCamInputSettings.h b/Engine/Plugins/Experimental/VirtualCamera/Source/VCamInput/Public/VCamInputSettings.h
--- a/Engine/Plugins/Experimental/VirtualCamera/Source/VCamInput/Public/VCamInputSettings.h
+++ b/Engine/Plugins/Experimental/VirtualCamera/Source/VCamInput/Public/VCamInputSettings.h
@@ -20,7 +20,7 @@ struct VCAMINPUT_API FVCamInputProfile
 	bool operator==(const FVCamInputProfile& OtherProfile) const;
 };
 
-UCLASS(config=Game, DefaultConfig)
+UCLASS(config=Game)
 class VCAMINPUT_API UVCamInputSettings : public UDeveloperSettings
 {
 	GENERATED_BODY()
@@ -43,4 +43,6 @@ public:
 
 	UFUNCTION(BlueprintCallable, Category="VCam Input")
 	TArray<FName> GetInputProfileNames() const;
+
+	virtual void PostInitProperties() override;
 };
\ No newline at end of file
diff --git a/Engine/Plugins/Experimental/VirtualCamera/Source/VCamInput/VCamInput.build.cs b/Engine/Plugins/Experimental/VirtualCamera/Source/VCamInput/VCamInput.build.cs
--- a/Engine/Plugins/Experimental/VirtualCamera/Source/VCamInput/VCamInput.build.cs
+++ b/Engine/Plugins/Experimental/VirtualCamera/Source/VCamInput/VCamInput.build.cs
@@ -25,6 +25,7 @@ public class VCamInput : ModuleRules
 				"CoreUObject",
 				"Engine",
 				"InputCore",
+				"Projects",
 				"Slate",
 				"SlateCore",
 				"DeveloperSettings"
