Allow Switch on Connections node to generate pins from a subwidget

diff --git a/Engine/Plugins/Experimental/VirtualCamera/Source/VCamBlueprintNodes/Private/SwitchNode/VCamK2Node_SwitchConnectionSystemBase.h b/Engine/Plugins/Experimental/VirtualCamera/Source/VCamBlueprintNodes/Private/SwitchNode/VCamK2Node_SwitchConnectionSystemBase.h
--- a/Engine/Plugins/Experimental/VirtualCamera/Source/VCamBlueprintNodes/Private/SwitchNode/VCamK2Node_SwitchConnectionSystemBase.h
+++ b/Engine/Plugins/Experimental/VirtualCamera/Source/VCamBlueprintNodes/Private/SwitchNode/VCamK2Node_SwitchConnectionSystemBase.h
@@ -34,6 +34,7 @@ public:
 	virtual void ReconstructNode() override;
 	virtual bool CanEverInsertExecutionPin() const override { return false; }
 	virtual bool CanEverRemoveExecutionPin() const override { return false; }
+	virtual bool ShouldShowNodeProperties() const override { return true; }
 	//~ End UK2Node Interface
 
 	// UK2Node_Switch Interface
diff --git a/Engine/Plugins/Experimental/VirtualCamera/Source/VCamBlueprintNodes/Private/SwitchNode/VCamK2Node_SwitchOnWidgetConnections.cpp b/Engine/Plugins/Experimental/VirtualCamera/Source/VCamBlueprintNodes/Private/SwitchNode/VCamK2Node_SwitchOnWidgetConnections.cpp
--- a/Engine/Plugins/Experimental/VirtualCamera/Source/VCamBlueprintNodes/Private/SwitchNode/VCamK2Node_SwitchOnWidgetConnections.cpp
+++ b/Engine/Plugins/Experimental/VirtualCamera/Source/VCamBlueprintNodes/Private/SwitchNode/VCamK2Node_SwitchOnWidgetConnections.cpp
@@ -24,13 +24,26 @@ bool UVCamK2Node_SwitchOnWidgetConnections::SupportsBlueprintClass(UClass* Class
 TArray<FName> UVCamK2Node_SwitchOnWidgetConnections::GetPinsToCreate() const
 {
 	TArray<FName> Result;
-	AccessBlueprintCDO([&Result](UVCamWidget* Widget)
+	AccessBlueprintCDO([this, &Result](UVCamWidget* Widget)
 	{
-		Widget->Connections.GenerateKeyArray(Result);
+		const UVCamWidget* ResolvedWidget = TargetWidget.ResolveVCamWidget(*Widget);
+		const UVCamWidget* WidgetToUse = ResolvedWidget ? ResolvedWidget : Widget;
+		WidgetToUse->Connections.GenerateKeyArray(Result);
 	});
 	return Result;
 }
 
+void UVCamK2Node_SwitchOnWidgetConnections::PostEditChangeProperty(FPropertyChangedEvent& PropertyChangedEvent)
+{
+	const FName PropertyName = (PropertyChangedEvent.MemberProperty != NULL) ? PropertyChangedEvent.MemberProperty->GetFName() : NAME_None;
+	if (PropertyName == GET_MEMBER_NAME_CHECKED(UVCamK2Node_SwitchOnWidgetConnections, TargetWidget))
+	{
+		ReconstructNode();
+	}
+	
+	Super::PostEditChangeProperty(PropertyChangedEvent);
+}
+
 void UVCamK2Node_SwitchOnWidgetConnections::AccessBlueprintCDO(TFunctionRef<void(UVCamWidget*)> Func) const
 {
 	if (!GetBlueprint())
diff --git a/Engine/Plugins/Experimental/VirtualCamera/Source/VCamBlueprintNodes/Private/SwitchNode/VCamK2Node_SwitchOnWidgetConnections.h b/Engine/Plugins/Experimental/VirtualCamera/Source/VCamBlueprintNodes/Private/SwitchNode/VCamK2Node_SwitchOnWidgetConnections.h
--- a/Engine/Plugins/Experimental/VirtualCamera/Source/VCamBlueprintNodes/Private/SwitchNode/VCamK2Node_SwitchOnWidgetConnections.h
+++ b/Engine/Plugins/Experimental/VirtualCamera/Source/VCamBlueprintNodes/Private/SwitchNode/VCamK2Node_SwitchOnWidgetConnections.h
@@ -4,6 +4,7 @@
 
 #include "CoreMinimal.h"
 #include "VCamK2Node_SwitchConnectionSystemBase.h"
+#include "Util/WidgetReference.h"
 #include "VCamK2Node_SwitchOnWidgetConnections.generated.h"
 
 class UVCamWidget;
@@ -15,6 +16,12 @@ UCLASS()
 class VCAMBLUEPRINTNODES_API UVCamK2Node_SwitchOnWidgetConnections : public UVCamK2Node_SwitchConnectionSystemBase
 {
 	GENERATED_BODY()
+public:
+
+	/** A child VCamWidget from which to generate the pins. If left None, this Blueprint's widget will be used. */
+	UPROPERTY(EditAnywhere, Category = PinOptions)
+	FVCamChildWidgetReference TargetWidget;
+	
 protected:
 
 	// UEdGraphNode interface
@@ -27,6 +34,10 @@ protected:
 	virtual TArray<FName> GetPinsToCreate() const override;
 	//~ End UVCamK2Node_SwitchConnectionSystemBase Interface
 
+	// UObject interface
+	virtual void PostEditChangeProperty(struct FPropertyChangedEvent& PropertyChangedEvent) override;
+	// End of UObject interface
+
 private:
 
 	void AccessBlueprintCDO(TFunctionRef<void(UVCamWidget*)> Func) const;
diff --git a/Engine/Plugins/Experimental/VirtualCamera/Source/VCamCore/Private/Util/WidgetReference.cpp b/Engine/Plugins/Experimental/VirtualCamera/Source/VCamCore/Private/Util/WidgetReference.cpp
--- a/Engine/Plugins/Experimental/VirtualCamera/Source/VCamCore/Private/Util/WidgetReference.cpp
+++ b/Engine/Plugins/Experimental/VirtualCamera/Source/VCamCore/Private/Util/WidgetReference.cpp
@@ -25,6 +25,7 @@ UWidget* FChildWidgetReference::ResolveWidget(UUserWidget& OwnerWidget) const
 
 	const bool bCanFindWidget = WidgetTree != nullptr;;
 	UE_CLOG(!bCanFindWidget, LogVCamCore, Warning, TEXT("Failed to get tree for widget %s"), *OwnerWidget.GetPathName());
+	ensureMsgf(ResolvedTemplate->GetTypedOuter<UWidgetTree>() && ResolvedTemplate->GetTypedOuter<UWidgetTree>()->HasAnyFlags(RF_DefaultSubObject), TEXT("Template does not point to any template widget!"));
 
 	const FName WidgetName = ResolvedTemplate->GetFName();
 	return bCanFindWidget
diff --git a/Engine/Plugins/Experimental/VirtualCamera/Source/VCamCoreEditor/Private/Customizations/WidgetReference/ChildWidgetReferenceCustomization.cpp b/Engine/Plugins/Experimental/VirtualCamera/Source/VCamCoreEditor/Private/Customizations/WidgetReference/ChildWidgetReferenceCustomization.cpp
--- a/Engine/Plugins/Experimental/VirtualCamera/Source/VCamCoreEditor/Private/Customizations/WidgetReference/ChildWidgetReferenceCustomization.cpp
+++ b/Engine/Plugins/Experimental/VirtualCamera/Source/VCamCoreEditor/Private/Customizations/WidgetReference/ChildWidgetReferenceCustomization.cpp
@@ -11,12 +11,41 @@
 #include "Blueprint/WidgetTree.h"
 #include "DetailWidgetRow.h"
 #include "IPropertyUtilities.h"
+#include "Blueprint/UserWidgetBlueprint.h"
 #include "Widgets/Text/STextBlock.h"
 
 #define LOCTEXT_NAMESPACE "FWidgetReferenceForBlueprintCustomization"
 
 namespace UE::VCamCoreEditor::Private
 {
+	static TWeakObjectPtr<UUserWidget> GetWidgetFromEditedObject(UObject* EditedObject)
+	{
+		if (!EditedObject)
+		{
+			return nullptr;
+		}
+		
+		if (UUserWidget* Result = Cast<UUserWidget>(EditedObject))
+		{
+			return Result;
+		}
+
+		// E.g. editing an instanced subobject, etc.
+		if (UUserWidget* OuterAsWidget = EditedObject->GetTypedOuter<UUserWidget>())
+		{
+			return OuterAsWidget;
+		}
+
+		// E.g. editing a Blueprint node
+		if (UUserWidgetBlueprint* OuterAsBlueprint = EditedObject->GetTypedOuter<UUserWidgetBlueprint>())
+		{
+			UObject* BlueprintDefaultObject = OuterAsBlueprint->GeneratedClass->GetDefaultObject();
+			return Cast<UUserWidget>(BlueprintDefaultObject);
+		}
+
+		return nullptr;
+	}
+	
 	TSharedRef<IPropertyTypeCustomization> FChildWidgetReferenceCustomization::MakeInstance()
 	{
 		return MakeShared<FChildWidgetReferenceCustomization>();
@@ -24,10 +53,10 @@ namespace UE::VCamCoreEditor::Private
 
 	void FChildWidgetReferenceCustomization::CustomizeHeader(TSharedRef<IPropertyHandle> PropertyHandle, FDetailWidgetRow& HeaderRow, IPropertyTypeCustomizationUtils& CustomizationUtils)
 	{
-		const TSharedPtr<IPropertyHandle> TemplateProperty = PropertyHandle->GetChildHandle(GET_MEMBER_NAME_CHECKED(FChildWidgetReference, Template));
-		
 		const TArray<TWeakObjectPtr<UObject>>& EditedObjects = CustomizationUtils.GetPropertyUtilities()->GetSelectedObjects();
-		TWeakObjectPtr<UUserWidget> EditedWidget = EditedObjects.Num() == 1 ? Cast<UUserWidget>(EditedObjects[0]) : nullptr;
+		const TWeakObjectPtr<UObject> EditedObject = EditedObjects.Num() == 1 ? EditedObjects[0] : nullptr;
+		TWeakObjectPtr<UUserWidget> EditedWidget = GetWidgetFromEditedObject(EditedObject.Get());
+		
 		// Always only search on the template
 		if (EditedWidget.IsValid() && !EditedWidget->HasAnyFlags(RF_ClassDefaultObject))
 		{
@@ -35,6 +64,7 @@ namespace UE::VCamCoreEditor::Private
 			EditedWidget = Cast<UUserWidget>(DefaultObject);
 		}
 		
+		const TSharedPtr<IPropertyHandle> TemplateProperty = PropertyHandle->GetChildHandle(GET_MEMBER_NAME_CHECKED(FChildWidgetReference, Template));
 		if (!EditedWidget.IsValid() || !ensure(TemplateProperty))
 		{
 			const FText Text = LOCTEXT("NotEditable", "No editable in this context");
