Implement a way to define hierarchy between modifiers procedurally

diff --git a/Engine/Plugins/Experimental/VirtualCamera/Source/VCamCore/Private/UI/VCamWidget.cpp b/Engine/Plugins/Experimental/VirtualCamera/Source/VCamCore/Private/UI/VCamWidget.cpp
--- a/Engine/Plugins/Experimental/VirtualCamera/Source/VCamCore/Private/UI/VCamWidget.cpp
+++ b/Engine/Plugins/Experimental/VirtualCamera/Source/VCamCore/Private/UI/VCamWidget.cpp
@@ -2,6 +2,7 @@
 
 #include "UI/VCamWidget.h"
 
+#include "LogVCamCore.h"
 #include "VCamComponent.h"
 
 UVCamWidget::UVCamWidget(const FObjectInitializer& ObjectInitializer)
@@ -14,6 +15,17 @@ UVCamWidget::UVCamWidget(const FObjectInitializer& ObjectInitializer)
 	});
 }
 
+void UVCamWidget::NativePreConstruct()
+{
+	// If we were created by a widget dynamically, let's use their VCamComponent so ReinitializeConnections will work
+	if (UVCamWidget* Outer = GetTypedOuter<UVCamWidget>(); !VCamComponent && Outer && Outer->VCamComponent)
+	{
+		VCamComponent = Outer->VCamComponent;
+	}
+	
+	Super::NativePreConstruct();
+}
+
 void UVCamWidget::NativeDestruct()
 {
 	Super::NativeDestruct();
@@ -43,6 +55,7 @@ bool UVCamWidget::ReinitializeConnections()
 {
 	if (!IsValid(VCamComponent))
 	{
+		UE_LOG(LogVCamCore, Error, TEXT("ReinitializeConnections failed because VCamComponet is not set (for widget %s of class %s)"), *GetPathName(), *GetClass()->GetPathName());
 		return false;
 	}
 	
diff --git a/Engine/Plugins/Experimental/VirtualCamera/Source/VCamCore/Public/UI/VCamWidget.h b/Engine/Plugins/Experimental/VirtualCamera/Source/VCamCore/Public/UI/VCamWidget.h
--- a/Engine/Plugins/Experimental/VirtualCamera/Source/VCamCore/Public/UI/VCamWidget.h
+++ b/Engine/Plugins/Experimental/VirtualCamera/Source/VCamCore/Public/UI/VCamWidget.h
@@ -69,6 +69,7 @@ public:
 	UVCamWidget(const FObjectInitializer& ObjectInitializer);
 
 	//~ Begin UWidget Interface
+	virtual void NativePreConstruct() override;
 	virtual void NativeDestruct() override;
 	//~ End UWidget Interface
 	
diff --git a/Engine/Plugins/Experimental/VirtualCamera/Source/VCamCoreEditor/Private/Factories/VCamWidgetFactory.cpp b/Engine/Plugins/Experimental/VirtualCamera/Source/VCamCoreEditor/Private/Factories/VCamWidgetFactory.cpp
--- a/Engine/Plugins/Experimental/VirtualCamera/Source/VCamCoreEditor/Private/Factories/VCamWidgetFactory.cpp
+++ b/Engine/Plugins/Experimental/VirtualCamera/Source/VCamCoreEditor/Private/Factories/VCamWidgetFactory.cpp
@@ -4,6 +4,7 @@
 #include "UI/VCamWidget.h"
 
 #include "AssetToolsModule.h"
+#include "IVCamCoreEditorModule.h"
 #include "Kismet2/BlueprintEditorUtils.h"
 #include "Kismet2/KismetEditorUtilities.h"
 #include "WidgetBlueprint.h"
@@ -50,8 +51,7 @@ UObject* UVCamWidgetFactory::FactoryCreateNew(UClass* Class, UObject* InParent,
 
 uint32 UVCamWidgetFactory::GetMenuCategories() const
 {
-	IAssetTools& AssetTools = FModuleManager::LoadModuleChecked<FAssetToolsModule>("AssetTools").Get();
-	return AssetTools.RegisterAdvancedAssetCategory("VirtualCamera", LOCTEXT("AssetCategoryName", "VCam"));
+	return UE::VCamCoreEditor::IVCamCoreEditorModule::Get().GetAdvancedAssetCategoryForVCam();
 }
 
 #undef LOCTEXT_NAMESPACE
diff --git a/Engine/Plugins/Experimental/VirtualCamera/Source/VCamCoreEditor/Private/VCamCoreEditorModule.cpp b/Engine/Plugins/Experimental/VirtualCamera/Source/VCamCoreEditor/Private/VCamCoreEditorModule.cpp
--- a/Engine/Plugins/Experimental/VirtualCamera/Source/VCamCoreEditor/Private/VCamCoreEditorModule.cpp
+++ b/Engine/Plugins/Experimental/VirtualCamera/Source/VCamCoreEditor/Private/VCamCoreEditorModule.cpp
@@ -2,6 +2,7 @@
 
 #include "VCamCoreEditorModule.h"
 
+#include "AssetToolsModule.h"
 #include "Customizations/ConnectionTargetSettingsTypeCustomization.h"
 #include "Customizations/OutputProvider/ConnectionRemapCustomization_StateSwitcher.h"
 #include "Customizations/OutputProvider/ConnectionRemapCustomization_VCamWidget.h"
@@ -30,6 +31,8 @@
 #include "Modules/ModuleManager.h"
 #include "PropertyEditorModule.h"
 
+#define LOCTEXT_NAMESPACE "FVCamCoreEditorModule"
+
 namespace UE::VCamCoreEditor::Private
 {
 	TSharedPtr<IConnectionRemapCustomization> FVCamCoreEditorModule::CreateConnectionRemapCustomization(TSubclassOf<UVCamWidget> Class) const
@@ -103,7 +106,13 @@ namespace UE::VCamCoreEditor::Private
 		// Any alive instances will be cleared when the details panel redraws
 		ConnectionRemappingCustomizationFactories.Remove(Class);
 	}
-	
+
+	uint32 FVCamCoreEditorModule::GetAdvancedAssetCategoryForVCam() const
+	{
+		IAssetTools& AssetTools = FModuleManager::LoadModuleChecked<FAssetToolsModule>("AssetTools").Get();
+		return AssetTools.RegisterAdvancedAssetCategory("VirtualCamera", LOCTEXT("AssetCategoryName", "VCam"));
+	}
+
 	void FVCamCoreEditorModule::RegisterAutoGeneratedDefaultEvents()
 	{
 		FKismetEditorUtilities::RegisterAutoGeneratedDefaultEvent(this, UVCamBlueprintModifier::StaticClass(), "OnInitialize");
@@ -190,4 +199,5 @@ namespace UE::VCamCoreEditor::Private
 	}
 }
 
-IMPLEMENT_MODULE(UE::VCamCoreEditor::Private::FVCamCoreEditorModule, VCamCoreEditor);
\ No newline at end of file
+IMPLEMENT_MODULE(UE::VCamCoreEditor::Private::FVCamCoreEditorModule, VCamCoreEditor);
+#undef LOCTEXT_NAMESPACE
\ No newline at end of file
diff --git a/Engine/Plugins/Experimental/VirtualCamera/Source/VCamCoreEditor/Private/VCamCoreEditorModule.h b/Engine/Plugins/Experimental/VirtualCamera/Source/VCamCoreEditor/Private/VCamCoreEditorModule.h
--- a/Engine/Plugins/Experimental/VirtualCamera/Source/VCamCoreEditor/Private/VCamCoreEditorModule.h
+++ b/Engine/Plugins/Experimental/VirtualCamera/Source/VCamCoreEditor/Private/VCamCoreEditorModule.h
@@ -33,6 +33,7 @@ namespace UE::VCamCoreEditor::Private
 		//~ Begin IVCamCoreEditorModule Interface
 		virtual void RegisterConnectionRemapCustomization(TSubclassOf<UVCamWidget> Class, FGetConnectionRemappingCustomization GetterDelegate) override;
 		virtual void UnregisterConnectionRemapCustomization(TSubclassOf<UVCamWidget> Class) override;
+		virtual uint32 GetAdvancedAssetCategoryForVCam() const override;
 		//~ Begin IVCamCoreEditorModule Interface
 
 	private:
diff --git a/Engine/Plugins/Experimental/VirtualCamera/Source/VCamCoreEditor/Public/IVCamCoreEditorModule.h b/Engine/Plugins/Experimental/VirtualCamera/Source/VCamCoreEditor/Public/IVCamCoreEditorModule.h
--- a/Engine/Plugins/Experimental/VirtualCamera/Source/VCamCoreEditor/Public/IVCamCoreEditorModule.h
+++ b/Engine/Plugins/Experimental/VirtualCamera/Source/VCamCoreEditor/Public/IVCamCoreEditorModule.h
@@ -35,6 +35,9 @@ namespace UE::VCamCoreEditor
 		 */
 		virtual void RegisterConnectionRemapCustomization(TSubclassOf<UVCamWidget> Class, FGetConnectionRemappingCustomization GetterDelegate) = 0;
 		virtual void UnregisterConnectionRemapCustomization(TSubclassOf<UVCamWidget> Class) = 0;
+
+		/** Gets the asset category that assets are grouped under. */
+		virtual uint32 GetAdvancedAssetCategoryForVCam() const = 0;
 		
 		virtual ~IVCamCoreEditorModule() override = default;
 	};
diff --git a/Engine/Plugins/Experimental/VirtualCamera/Source/VCamExtensions/Private/Hierarchies/ClassBasedModifierHierarchyRules.cpp b/Engine/Plugins/Experimental/VirtualCamera/Source/VCamExtensions/Private/Hierarchies/ClassBasedModifierHierarchyRules.cpp
new file mode 100644
--- /dev/null
+++ b/Engine/Plugins/Experimental/VirtualCamera/Source/VCamExtensions/Private/Hierarchies/ClassBasedModifierHierarchyRules.cpp
@@ -0,0 +1,185 @@
+// Copyright Epic Games, Inc. All Rights Reserved.
+
+#include "Hierarchies/ClassBasedModifierHierarchyRules.h"
+
+#include "VCamComponent.h"
+
+#include "Algo/AllOf.h"
+
+void UClassBasedModifierGroup::PostInitProperties()
+{
+	GroupName = GetFName();
+	Super::PostInitProperties();
+}
+
+UClassBasedModifierHierarchyRules::UClassBasedModifierHierarchyRules()
+{
+	RootGroup = CreateDefaultSubobject<UClassBasedModifierGroup>(TEXT("Root"));
+}
+
+FName UClassBasedModifierHierarchyRules::GetRootGroup_Implementation() const
+{
+	return ensure(RootGroup)
+		? RootGroup->GroupName
+		: NAME_None;
+}
+
+bool UClassBasedModifierHierarchyRules::GetGroupOfModifier_Implementation(UVCamModifier* Modifier, FName& Group) const
+{
+	if (!Modifier)
+	{
+		return false;
+	}
+
+	if (UClassBasedModifierGroup* MatchingGroup = FindBestMatchFor(*Modifier))
+	{
+		Group = MatchingGroup->GroupName;
+		return true;
+	}
+	return false;
+}
+
+TSet<FName> UClassBasedModifierHierarchyRules::GetChildGroups_Implementation(FName ParentGroup) const
+{
+	if (const UClassBasedModifierGroup* FoundGroup = FindGroupByName(ParentGroup))
+	{
+		TSet<FName> Result;
+		Algo::TransformIf(FoundGroup->Children, Result, [](const TObjectPtr<UClassBasedModifierGroup>& Child){ return Child != nullptr; }, [](const TObjectPtr<UClassBasedModifierGroup>& Child){ return Child->GroupName; });
+		return Result;
+	}
+	
+	return {};
+}
+
+TSet<UVCamModifier*> UClassBasedModifierHierarchyRules::GetModifiersInGroup_Implementation(UVCamComponent* Component, FName GroupName) const
+{
+	UClassBasedModifierGroup* FoundGroup = FindGroupByName(GroupName);
+	return FoundGroup
+		? EnumerateModifiersInGroup(*FoundGroup, *Component)
+		: TSet<UVCamModifier*>{};
+}
+
+#if WITH_EDITOR
+void UClassBasedModifierHierarchyRules::PostEditChangeProperty(FPropertyChangedEvent& PropertyChangedEvent)
+{
+	if (PropertyChangedEvent.Property
+		&& PropertyChangedEvent.Property->GetFName() == GET_MEMBER_NAME_CHECKED(UClassBasedModifierGroup, GroupName))
+	{
+		// TODO: ensure unique name
+	}
+	
+	Super::PostEditChangeProperty(PropertyChangedEvent);
+}
+#endif
+
+UClassBasedModifierGroup* UClassBasedModifierHierarchyRules::FindGroupByName(FName GroupName) const
+{
+	UClassBasedModifierGroup* FoundGroup = nullptr;
+	ForEachGroup([GroupName, &FoundGroup](UClassBasedModifierGroup& Group)
+	{
+		if (Group.GroupName == GroupName)
+		{
+			FoundGroup = &Group;
+			return EBreakBehavior::Break;
+		}
+		return EBreakBehavior::Continue;
+	});
+	return FoundGroup;
+}
+
+TSet<UVCamModifier*> UClassBasedModifierHierarchyRules::EnumerateModifiersInGroup(UClassBasedModifierGroup& Group, UVCamComponent& Component) const
+{
+	TSet<UVCamModifier*> AllModifiers;
+	TArray<UVCamModifier*> ModifersOfSameClass;
+		
+	for (const TSubclassOf<UVCamModifier>& ModifierClass : Group.ModifierClasses)
+	{
+		ModifersOfSameClass.Empty(Component.GetNumberOfModifiers());
+		Component.GetModifiersByClass(ModifierClass, ModifersOfSameClass);
+		AllModifiers.Append(ModifersOfSameClass);
+	}
+
+	// GetModifiersByClass gets all modifiers that are a subclass.
+	// If a more specific class is specified on a different node, then this modifier should not be put into the current group
+	
+	// This implementation has a really bad complexity of O(n*m) ~ O(n^2) ... where n = number of nodes, m = AllModifiers.Num()
+	// We could improve it by caching with a TMap which nodes store what classes. 
+	for (auto AllModifierIt = AllModifiers.CreateIterator(); AllModifierIt; ++AllModifierIt)
+	{
+		if (FindBestMatchFor(**AllModifierIt) != &Group)
+		{
+			AllModifierIt.RemoveCurrent();
+		}
+	}
+
+	return AllModifiers;
+}
+
+UClassBasedModifierGroup* UClassBasedModifierHierarchyRules::FindBestMatchFor(UVCamModifier& Modifier) const
+{
+	TArray<TPair<TSubclassOf<UVCamModifier>, UClassBasedModifierGroup&>> MatchingGroups;
+	ForEachGroup([&Modifier, &MatchingGroups](UClassBasedModifierGroup& Group)
+	{
+		for (const TSubclassOf<UVCamModifier>& Class : Group.ModifierClasses)
+		{
+			if (Modifier.IsA(Class))
+			{
+				MatchingGroups.Emplace(Class, Group);
+			}
+		}
+		
+		return EBreakBehavior::Continue;
+	});
+
+	auto GetInheritanceDistance = [](TSubclassOf<UVCamModifier> Base, TSubclassOf<UVCamModifier> Subclass)
+	{
+		int32 Distance = 0;
+		for (UClass* Current = Subclass; Current != Base; Current = Current->GetSuperClass())
+		{
+			++Distance;
+		}
+		return Distance;
+	};
+	
+	uint32 CurrentLowestDistance = GetInheritanceDistance(UVCamModifier::StaticClass(), Modifier.GetClass());
+	UClassBasedModifierGroup* ClosestGroup = nullptr;
+	for (const TPair<TSubclassOf<UVCamModifier>, UClassBasedModifierGroup&>& Pair : MatchingGroups)
+	{
+		const uint32 CurrentDistance = GetInheritanceDistance(Pair.Key, Modifier.GetClass());
+		if (CurrentLowestDistance > CurrentDistance)
+		{
+			CurrentLowestDistance = CurrentDistance;
+			ClosestGroup = &Pair.Value;
+		}
+	}
+	return ClosestGroup;
+}
+
+void UClassBasedModifierHierarchyRules::ForEachGroup(TFunctionRef<EBreakBehavior(UClassBasedModifierGroup& Group)> Callback) const
+{
+	if (!ensure(RootGroup))
+	{
+		return;
+	}
+
+	TQueue<UClassBasedModifierGroup*> Queue;
+	Queue.Enqueue(RootGroup);
+
+	UClassBasedModifierGroup* Current = nullptr;
+	while (Queue.Dequeue(Current))
+	{
+		const EBreakBehavior BreakBehavior = Callback(*Current);
+		if (BreakBehavior == EBreakBehavior::Break)
+		{
+			return;
+		}
+		
+		for (UClassBasedModifierGroup* Child : Current->Children)
+		{
+			if (Child)
+			{
+				Queue.Enqueue(Child);
+			}
+		}
+	}
+}
diff --git a/Engine/Plugins/Experimental/VirtualCamera/Source/VCamExtensions/Private/Hierarchies/ModifierHierarchyAsset.cpp b/Engine/Plugins/Experimental/VirtualCamera/Source/VCamExtensions/Private/Hierarchies/ModifierHierarchyAsset.cpp
new file mode 100644
--- /dev/null
+++ b/Engine/Plugins/Experimental/VirtualCamera/Source/VCamExtensions/Private/Hierarchies/ModifierHierarchyAsset.cpp
@@ -0,0 +1,26 @@
+// Copyright Epic Games, Inc. All Rights Reserved.
+
+#include "Hierarchies/ModifierHierarchyAsset.h"
+
+#include "Hierarchies/ClassBasedModifierHierarchyRules.h"
+#include "Hierarchies/ModifierHierarchyRules.h"
+
+FName UModifierHierarchyAsset::GetRootGroup() const
+{
+	return Rules ? Rules->GetRootGroup() : NAME_None;
+}
+
+bool UModifierHierarchyAsset::GetGroupOfModifier(UVCamModifier* Modifier, FName& Group) const
+{
+	return Rules && Rules->GetGroupOfModifier(Modifier, Group);
+}
+
+TSet<FName> UModifierHierarchyAsset::GetChildGroups_Implementation(FName ParentGroup) const
+{
+	return Rules ? Rules->GetChildGroups(ParentGroup) : TSet<FName>{};
+}
+
+TSet<UVCamModifier*> UModifierHierarchyAsset::GetModifiersInGroup(UVCamComponent* Component, FName GroupName) const
+{
+	return Rules ? Rules->GetModifiersInGroup(Component, GroupName) : TSet<UVCamModifier*>{};
+}
diff --git a/Engine/Plugins/Experimental/VirtualCamera/Source/VCamExtensions/Private/Hierarchies/ModifierHierarchyRules.cpp b/Engine/Plugins/Experimental/VirtualCamera/Source/VCamExtensions/Private/Hierarchies/ModifierHierarchyRules.cpp
new file mode 100644
--- /dev/null
+++ b/Engine/Plugins/Experimental/VirtualCamera/Source/VCamExtensions/Private/Hierarchies/ModifierHierarchyRules.cpp
@@ -0,0 +1,27 @@
+// Copyright Epic Games, Inc. All Rights Reserved.
+
+#include "Hierarchies/ModifierHierarchyRules.h"
+
+FName UModifierHierarchyRules::GetRootGroup_Implementation() const
+{
+	unimplemented();
+	return {};
+}
+
+bool UModifierHierarchyRules::GetGroupOfModifier_Implementation(UVCamModifier* Modifier, FName& Group) const
+{
+	unimplemented();
+	return false;
+}
+
+TSet<FName> UModifierHierarchyRules::GetChildGroups_Implementation(FName ParentGroup) const
+{
+	unimplemented();
+	return {};
+}
+
+TSet<UVCamModifier*> UModifierHierarchyRules::GetModifiersInGroup_Implementation(UVCamComponent* Component, FName GroupName) const
+{
+	unimplemented();
+	return {};
+}
diff --git a/Engine/Plugins/Experimental/VirtualCamera/Source/VCamExtensions/Private/Styling/ClassBasedWidgetStyleDefinitions.cpp b/Engine/Plugins/Experimental/VirtualCamera/Source/VCamExtensions/Private/Styling/ClassBasedWidgetStyleDefinitions.cpp
new file mode 100644
--- /dev/null
+++ b/Engine/Plugins/Experimental/VirtualCamera/Source/VCamExtensions/Private/Styling/ClassBasedWidgetStyleDefinitions.cpp
@@ -0,0 +1,86 @@
+// Copyright Epic Games, Inc. All Rights Reserved.
+
+#include "Styling/ClassBasedWidgetStyleDefinitions.h"
+
+#include "VCamComponent.h"
+#include "Algo/ForEach.h"
+
+DEFINE_LOG_CATEGORY_STATIC(LogClassBasedModifierMetaDataRules, All, All);
+
+namespace UE::VCamExtensions::Private
+{
+	static TArray<UWidgetStyleData*> AccumulateMetaDataRecursively(
+		const TMap<TSubclassOf<UVCamModifier>, FPerModifierClassWidgetSytleData>& Config,
+		UVCamComponent* VCamComponent, 
+		FName ModifierId,
+		TFunctionRef<const FWidgetStyleDataConfig*(const FPerModifierClassWidgetSytleData& ClassConfig)> RetrieveConfig)
+	{
+		TArray<UWidgetStyleData*> Result;
+		if (!IsValid(VCamComponent))
+		{
+			return Result;
+		}
+
+		UVCamModifier* Modifier = VCamComponent->GetModifierByName(ModifierId);
+		if (!Modifier)
+		{
+			UE_LOG(LogClassBasedModifierMetaDataRules, Warning, TEXT("Unknown modifier %s (on component %s)"), *ModifierId.ToString(), *VCamComponent->GetPathName());
+			return Result;
+		}
+
+		for (UClass* ModifierClass = Modifier->GetClass(); ModifierClass != UVCamModifier::StaticClass()->GetSuperClass(); ModifierClass = ModifierClass->GetSuperClass())
+		{
+			const FPerModifierClassWidgetSytleData* ClassConfig = Config.Find(ModifierClass);
+			if (!ClassConfig)
+			{
+				continue;
+			}
+			
+			const FWidgetStyleDataConfig* MetaDataConfig = RetrieveConfig(*ClassConfig);
+			if (!MetaDataConfig)
+			{
+				continue;
+			}
+
+			Algo::ForEach(MetaDataConfig->ModifierMetaData, [&Result](const TObjectPtr<UWidgetStyleData>& MetaData) { Result.Add(MetaData); });
+			if (!MetaDataConfig->bInherit)
+			{
+				break;
+			}
+		}
+		return Result;
+	}
+}
+
+
+TArray<UWidgetStyleData*> UClassBasedWidgetStyleDefinitions::GetStylesForModifier_Implementation(UVCamComponent* VCamComponent, FName ModifierId) const
+{
+	return UE::VCamExtensions::Private::AccumulateMetaDataRecursively(
+		Config,
+		VCamComponent,
+		ModifierId,
+		[](const FPerModifierClassWidgetSytleData& ClassConfig)
+		{
+			return &ClassConfig.ModifierStyles;
+		});
+}
+
+TArray<UWidgetStyleData*> UClassBasedWidgetStyleDefinitions::GetStylesForCategoryInModifier_Implementation(UVCamComponent* VCamComponent, FName ModifierId, FName ConnectionPointId) const
+{
+	return UE::VCamExtensions::Private::AccumulateMetaDataRecursively(
+		Config,
+		VCamComponent,
+		ModifierId,
+		[ConnectionPointId](const FPerModifierClassWidgetSytleData& ClassConfig)
+		{
+			return ClassConfig.CategorizedStyles.Find(ConnectionPointId);
+		});
+}
+
+TArray<UWidgetStyleData*> UClassBasedWidgetStyleDefinitions::GetStylesForCategoryWithoutModifier_Implementation(FName Category) const
+{
+	const FWidgetStyleDataArray* Result = CategoriesWithoutModifier.Find(Category);
+	return Result
+		? Result->Styles
+		: TArray<UWidgetStyleData*>{};
+}
diff --git a/Engine/Plugins/Experimental/VirtualCamera/Source/VCamExtensions/Private/Styling/ModifierBoundWidgetStyleDefinitions.cpp b/Engine/Plugins/Experimental/VirtualCamera/Source/VCamExtensions/Private/Styling/ModifierBoundWidgetStyleDefinitions.cpp
new file mode 100644
--- /dev/null
+++ b/Engine/Plugins/Experimental/VirtualCamera/Source/VCamExtensions/Private/Styling/ModifierBoundWidgetStyleDefinitions.cpp
@@ -0,0 +1,62 @@
+// Copyright Epic Games, Inc. All Rights Reserved.
+
+#include "Styling/ModifierBoundWidgetStyleDefinitions.h"
+
+#include "Styling/WidgetStyleData.h"
+
+namespace UE::VCamExtensions::Private
+{
+	static UWidgetStyleData* Filter(const TArray<UWidgetStyleData*>& MetaDataArray, TSubclassOf<UWidgetStyleData> Class)
+	{
+		for (UWidgetStyleData* MetaData : MetaDataArray)
+		{
+			if (MetaData && MetaData->IsA(Class))
+			{
+				return MetaData;
+			}
+		}
+		return nullptr;
+	}
+}
+
+TArray<UWidgetStyleData*> UModifierBoundWidgetStyleDefinitions::GetStylesForModifier_Implementation(UVCamComponent* VCamComponent, FName ModifierId) const
+{
+	unimplemented();
+	return {};
+}
+
+TArray<UWidgetStyleData*> UModifierBoundWidgetStyleDefinitions::GetStylesForCategoryInModifier_Implementation(UVCamComponent* VCamComponent, FName ModifierId, FName Category) const
+{
+	unimplemented();
+	return {};
+}
+
+TArray<UWidgetStyleData*> UModifierBoundWidgetStyleDefinitions::GetStylesForCategoryWithoutModifier_Implementation(FName Category) const
+{
+	unimplemented();
+	return {};
+}
+
+UWidgetStyleData* UModifierBoundWidgetStyleDefinitions::GetStyleForModifierByClass(UVCamComponent* VCamComponent, FName ModifierId, TSubclassOf<UWidgetStyleData> Class) const
+{
+	return UE::VCamExtensions::Private::Filter(
+		GetStylesForModifier(VCamComponent, ModifierId),
+		Class
+		);
+}
+
+UWidgetStyleData* UModifierBoundWidgetStyleDefinitions::GetStyleForCategoryByClassInModifier(UVCamComponent* VCamComponent, FName ModifierId, FName Category, TSubclassOf<UWidgetStyleData> Class) const
+{
+	return UE::VCamExtensions::Private::Filter(
+		GetStylesForCategoryInModifier(VCamComponent, ModifierId, Category),
+		Class
+		);
+}
+
+UWidgetStyleData* UModifierBoundWidgetStyleDefinitions::GetStyleForCategoryWithoutModifier(FName Category, TSubclassOf<UWidgetStyleData> Class) const
+{
+	return UE::VCamExtensions::Private::Filter(
+			GetStylesForCategoryWithoutModifier(Category),
+			Class
+			);
+}
\ No newline at end of file
diff --git a/Engine/Plugins/Experimental/VirtualCamera/Source/VCamExtensions/Private/Styling/ModifierBoundWidgetStylesAsset.cpp b/Engine/Plugins/Experimental/VirtualCamera/Source/VCamExtensions/Private/Styling/ModifierBoundWidgetStylesAsset.cpp
new file mode 100644
--- /dev/null
+++ b/Engine/Plugins/Experimental/VirtualCamera/Source/VCamExtensions/Private/Styling/ModifierBoundWidgetStylesAsset.cpp
@@ -0,0 +1,36 @@
+// Copyright Epic Games, Inc. All Rights Reserved.
+
+#include "Styling/ModifierBoundWidgetStylesAsset.h"
+
+#include "Styling/ClassBasedWidgetStyleDefinitions.h"
+#include "Styling/ModifierBoundWidgetStyleDefinitions.h"
+
+TArray<UWidgetStyleData*> UModifierBoundWidgetStylesAsset::GetStylesForModifier(UVCamComponent* VCamComponent, FName ModifierId) const
+{
+	return Rules ? Rules->GetStylesForModifier(VCamComponent, ModifierId) : TArray<UWidgetStyleData*>{};
+}
+
+TArray<UWidgetStyleData*> UModifierBoundWidgetStylesAsset::GetStylesForCategoryInModifier(UVCamComponent* VCamComponent, FName ModifierId, FName Category) const
+{
+	return Rules ? Rules->GetStylesForCategoryInModifier(VCamComponent, ModifierId, Category) : TArray<UWidgetStyleData*>{};
+}
+
+TArray<UWidgetStyleData*> UModifierBoundWidgetStylesAsset::GetStylesForCategoryWithoutModifier(FName Category) const
+{
+	return Rules ? Rules->GetStylesForCategoryWithoutModifier(Category) : TArray<UWidgetStyleData*>{};
+}
+
+UWidgetStyleData* UModifierBoundWidgetStylesAsset::GetStyleForModifierByClass(UVCamComponent* VCamComponent, FName ModifierId, TSubclassOf<UWidgetStyleData> Class) const
+{
+	return Rules ? Rules->GetStyleForModifierByClass(VCamComponent, ModifierId, Class) : nullptr;
+}
+
+UWidgetStyleData* UModifierBoundWidgetStylesAsset::GetStyleForCategoryByClassInModifier(UVCamComponent* VCamComponent, FName ModifierId, FName Category, TSubclassOf<UWidgetStyleData> Class) const
+{
+	return Rules ? Rules->GetStyleForCategoryByClassInModifier(VCamComponent, ModifierId, Category, Class) : nullptr;
+}
+
+UWidgetStyleData* UModifierBoundWidgetStylesAsset::GetStyleForCategoryWithoutModifier(FName Category, TSubclassOf<UWidgetStyleData> Class) const
+{
+	return Rules ? Rules->GetStyleForCategoryWithoutModifier(Category, Class) : nullptr;
+}
diff --git a/Engine/Plugins/Experimental/VirtualCamera/Source/VCamExtensions/Private/VCamExtensions.cpp b/Engine/Plugins/Experimental/VirtualCamera/Source/VCamExtensions/Private/VCamExtensions.cpp
new file mode 100644
--- /dev/null
+++ b/Engine/Plugins/Experimental/VirtualCamera/Source/VCamExtensions/Private/VCamExtensions.cpp
@@ -0,0 +1,20 @@
+// Copyright Epic Games, Inc. All Rights Reserved.
+
+#include "Modules/ModuleInterface.h"
+#include "Modules/ModuleManager.h"
+
+namespace UE::VCamExtensions::Private
+{
+	class FVCamExtensionsModule : public IModuleInterface
+	{
+	public:
+
+		virtual void StartupModule() override
+		{}
+	
+		virtual void ShutdownModule() override
+		{}
+	};
+}
+
+IMPLEMENT_MODULE(UE::VCamExtensions::Private::FVCamExtensionsModule, VCamExtensions);
\ No newline at end of file
diff --git a/Engine/Plugins/Experimental/VirtualCamera/Source/VCamExtensions/Private/VCamExtensionsLog.cpp b/Engine/Plugins/Experimental/VirtualCamera/Source/VCamExtensions/Private/VCamExtensionsLog.cpp
new file mode 100644
--- /dev/null
+++ b/Engine/Plugins/Experimental/VirtualCamera/Source/VCamExtensions/Private/VCamExtensionsLog.cpp
@@ -0,0 +1,5 @@
+// Copyright Epic Games, Inc. All Rights Reserved.
+
+#include "VCamExtensionsLog.h"
+
+DEFINE_LOG_CATEGORY(LogVCamExtensions);
\ No newline at end of file
diff --git a/Engine/Plugins/Experimental/VirtualCamera/Source/VCamExtensions/Private/VCamExtensionsLog.h b/Engine/Plugins/Experimental/VirtualCamera/Source/VCamExtensions/Private/VCamExtensionsLog.h
new file mode 100644
--- /dev/null
+++ b/Engine/Plugins/Experimental/VirtualCamera/Source/VCamExtensions/Private/VCamExtensionsLog.h
@@ -0,0 +1,7 @@
+// Copyright Epic Games, Inc. All Rights Reserved.
+
+#pragma once
+
+#include "Logging/LogMacros.h"
+
+DECLARE_LOG_CATEGORY_EXTERN(LogVCamExtensions, Log, All);
\ No newline at end of file
diff --git a/Engine/Plugins/Experimental/VirtualCamera/Source/VCamExtensions/Public/Hierarchies/ClassBasedModifierHierarchyRules.h b/Engine/Plugins/Experimental/VirtualCamera/Source/VCamExtensions/Public/Hierarchies/ClassBasedModifierHierarchyRules.h
new file mode 100644
--- /dev/null
+++ b/Engine/Plugins/Experimental/VirtualCamera/Source/VCamExtensions/Public/Hierarchies/ClassBasedModifierHierarchyRules.h
@@ -0,0 +1,70 @@
+// Copyright Epic Games, Inc. All Rights Reserved.
+
+#pragma once
+
+#include "ModifierHierarchyRules.h"
+#include "Modifier/VCamModifier.h"
+#include "Templates/SubclassOf.h"
+#include "ClassBasedModifierHierarchyRules.generated.h"
+
+class UVCamComponent;
+
+UCLASS(EditInlineNew)
+class VCAMEXTENSIONS_API UClassBasedModifierGroup : public UObject
+{
+	GENERATED_BODY()
+public:
+
+	UPROPERTY(EditAnywhere, Category = "Virtual Camera")
+	FName GroupName;
+	
+	UPROPERTY(EditAnywhere, Category = "Virtual Camera")
+	TSet<TSubclassOf<UVCamModifier>> ModifierClasses;
+
+	UPROPERTY(EditAnywhere, Instanced, Category = "Virtual Camera")
+	TArray<TObjectPtr<UClassBasedModifierGroup>> Children;
+
+	virtual void PostInitProperties() override;
+};
+
+/**
+ * 
+ */
+UCLASS()
+class VCAMEXTENSIONS_API UClassBasedModifierHierarchyRules : public UModifierHierarchyRules
+{
+	GENERATED_BODY()
+public:
+
+	UClassBasedModifierHierarchyRules();
+
+	//~ Begin UModifierHierarchyRules Interface
+	virtual FName GetRootGroup_Implementation() const override;
+	virtual bool GetGroupOfModifier_Implementation(UVCamModifier* Modifier, FName& Group) const override;
+	virtual TSet<FName> GetChildGroups_Implementation(FName ParentGroup) const override;
+	virtual TSet<UVCamModifier*> GetModifiersInGroup_Implementation(UVCamComponent* Component, FName GroupName) const override;
+	//~ End UModifierHierarchyRules Interface
+
+#if WITH_EDITOR
+	//~ Begin UObject Interface
+	virtual void PostEditChangeProperty(FPropertyChangedEvent& PropertyChangedEvent) override;
+	//~ End UObject Interface
+#endif
+
+private:
+	
+	UPROPERTY(EditAnywhere, Instanced, Category = "Virtual Camera", meta = (NoResetToDefault))
+	TObjectPtr<UClassBasedModifierGroup> RootGroup;
+
+	enum class EBreakBehavior
+	{
+		Continue,
+		Break
+	};
+
+	UClassBasedModifierGroup* FindGroupByName(FName GroupName) const;
+	TSet<UVCamModifier*> EnumerateModifiersInGroup(UClassBasedModifierGroup& Group, UVCamComponent& Component) const;
+	
+	UClassBasedModifierGroup* FindBestMatchFor(UVCamModifier& Modifier) const;
+	void ForEachGroup(TFunctionRef<EBreakBehavior(UClassBasedModifierGroup& Group)> Callback) const;
+};
diff --git a/Engine/Plugins/Experimental/VirtualCamera/Source/VCamExtensions/Public/Hierarchies/ModifierHierarchyAsset.h b/Engine/Plugins/Experimental/VirtualCamera/Source/VCamExtensions/Public/Hierarchies/ModifierHierarchyAsset.h
new file mode 100644
--- /dev/null
+++ b/Engine/Plugins/Experimental/VirtualCamera/Source/VCamExtensions/Public/Hierarchies/ModifierHierarchyAsset.h
@@ -0,0 +1,50 @@
+// Copyright Epic Games, Inc. All Rights Reserved.
+
+#pragma once
+
+#include "UObject/Object.h"
+#include "ModifierHierarchyAsset.generated.h"
+
+class UModifierHierarchyRules;
+class UVCamComponent;
+class UVCamModifier;
+
+/**
+ * An asset intended to be referenced by Slate widgets.
+ *
+ * Defines a tree hierarchy. Nodes are called groups.
+ * A group consists of modifiers and (sub) groups.
+ * 
+ * An example use case is if you want to have a button menu which should procedurally generate sub-button menus depending
+ * on the modifiers in the component. 
+ * One group could be a Lens group which groups together modifiers that e.g. modify focal distance, filmback, and FOV.
+ * The rules are defined generically enough so as modifiers are added or removed, the groupings also update accordingly.
+ */
+UCLASS(BlueprintType)
+class VCAMEXTENSIONS_API UModifierHierarchyAsset : public UObject
+{
+	GENERATED_BODY()
+public:
+
+	UPROPERTY(EditAnywhere, BlueprintReadOnly, Instanced, Category = "Virtual Camera")
+	TObjectPtr<UModifierHierarchyRules> Rules;
+
+	/** Gets the root of the tree. */
+	UFUNCTION(BlueprintPure, Category = "Virtual Camera")
+	FName GetRootGroup() const;
+
+	/**
+	 * Gets the group the modifier belongs to.
+	 * @return True if the modifier belongs to any group
+	 */
+	UFUNCTION(BlueprintPure, Category = "Virtual Camera")
+	bool GetGroupOfModifier(UVCamModifier* Modifier, FName& Group) const;
+
+	/** Gets the child groups of the given group. */
+	UFUNCTION(BlueprintPure, BlueprintNativeEvent, Category = "Virtual Camera")
+	TSet<FName> GetChildGroups(FName ParentGroup) const;
+
+	/** Gets all the modifiers on the component that belong in the given group. */
+	UFUNCTION(BlueprintPure, Category = "Virtual Camera")
+	TSet<UVCamModifier*> GetModifiersInGroup(UVCamComponent* Component, FName GroupName) const;
+};
\ No newline at end of file
diff --git a/Engine/Plugins/Experimental/VirtualCamera/Source/VCamExtensions/Public/Hierarchies/ModifierHierarchyRules.h b/Engine/Plugins/Experimental/VirtualCamera/Source/VCamExtensions/Public/Hierarchies/ModifierHierarchyRules.h
new file mode 100644
--- /dev/null
+++ b/Engine/Plugins/Experimental/VirtualCamera/Source/VCamExtensions/Public/Hierarchies/ModifierHierarchyRules.h
@@ -0,0 +1,44 @@
+// Copyright Epic Games, Inc. All Rights Reserved.
+
+#pragma once
+
+#include "UObject/Object.h"
+#include "ModifierHierarchyRules.generated.h"
+
+class UVCamComponent;
+class UVCamModifier;
+
+/**
+ * Defines a tree hierarchy. Each node is called a group.
+ * A group consists of modifiers and (sub) groups.
+ * 
+ * An example use case is if you want to have a button menu which should procedurally generate sub-button menus depending
+ * on the modifiers in the component. 
+ * One group could be a Lens group which groups together modifiers that e.g. modify focal distance, filmback, and FOV.
+ * The rules are defined generically enough so as modifiers are added or removed, the groupings also update accordingly.
+ */
+UCLASS(Abstract, Blueprintable, EditInlineNew)
+class VCAMEXTENSIONS_API UModifierHierarchyRules : public UObject
+{
+	GENERATED_BODY()
+public:
+
+	/** Gets the root of the tree. */
+	UFUNCTION(BlueprintPure, BlueprintNativeEvent, Category = "Virtual Camera")
+	FName GetRootGroup() const;
+
+	/**
+	 * Gets the group the modifier belongs to.
+	 * @return True if the modifier belongs to any group
+	 */
+	UFUNCTION(BlueprintPure, BlueprintNativeEvent, Category = "Virtual Camera")
+	bool GetGroupOfModifier(UVCamModifier* Modifier, FName& Group) const;
+
+	/** Gets the child groups of the given group. */
+	UFUNCTION(BlueprintPure, BlueprintNativeEvent, Category = "Virtual Camera")
+	TSet<FName> GetChildGroups(FName ParentGroup) const;
+
+	/** Gets all the modifiers on the component that belong in the given group. */
+	UFUNCTION(BlueprintPure, BlueprintNativeEvent, Category = "Virtual Camera")
+	TSet<UVCamModifier*> GetModifiersInGroup(UVCamComponent* Component, FName GroupName) const;
+};
diff --git a/Engine/Plugins/Experimental/VirtualCamera/Source/VCamExtensions/Public/Styling/ButtonWidgetStyleData.h b/Engine/Plugins/Experimental/VirtualCamera/Source/VCamExtensions/Public/Styling/ButtonWidgetStyleData.h
new file mode 100644
--- /dev/null
+++ b/Engine/Plugins/Experimental/VirtualCamera/Source/VCamExtensions/Public/Styling/ButtonWidgetStyleData.h
@@ -0,0 +1,17 @@
+// Copyright Epic Games, Inc. All Rights Reserved.
+
+#pragma once
+
+#include "WidgetStyleData.h"
+#include "ButtonWidgetStyleData.generated.h"
+
+/** Associates a button style for widgets that want to display a button for a modifier / connection point. */
+UCLASS()
+class VCAMEXTENSIONS_API UButtonWidgetStyleData : public UWidgetStyleData
+{
+	GENERATED_BODY()
+public:
+
+	UPROPERTY(EditAnywhere, BlueprintReadWrite, Category = "Virtual Camera")
+	FButtonStyle ButtonStyle;
+};
\ No newline at end of file
diff --git a/Engine/Plugins/Experimental/VirtualCamera/Source/VCamExtensions/Public/Styling/ClassBasedWidgetStyleDefinitions.h b/Engine/Plugins/Experimental/VirtualCamera/Source/VCamExtensions/Public/Styling/ClassBasedWidgetStyleDefinitions.h
new file mode 100644
--- /dev/null
+++ b/Engine/Plugins/Experimental/VirtualCamera/Source/VCamExtensions/Public/Styling/ClassBasedWidgetStyleDefinitions.h
@@ -0,0 +1,68 @@
+// Copyright Epic Games, Inc. All Rights Reserved.
+
+#pragma once
+
+#include "ModifierBoundWidgetStyleDefinitions.h"
+#include "Modifier/VCamModifier.h"
+#include "ClassBasedWidgetStyleDefinitions.generated.h"
+
+class UWidgetStyleData;
+
+USTRUCT()
+struct VCAMEXTENSIONS_API FWidgetStyleDataArray
+{
+	GENERATED_BODY()
+	
+	UPROPERTY(EditAnywhere, Instanced, Category = "Virtual Camera")
+	TArray<TObjectPtr<UWidgetStyleData>> Styles;
+};
+
+USTRUCT()
+struct VCAMEXTENSIONS_API FWidgetStyleDataConfig
+{
+	GENERATED_BODY()
+
+	UPROPERTY(EditAnywhere, Instanced, Category = "Virtual Camera")
+	TArray<TObjectPtr<UWidgetStyleData>> ModifierMetaData;
+	
+	/** Whether to inherit the meta data of modifier's parent class(es). */
+	UPROPERTY(EditAnywhere, Category = "Virtual Camera", AdvancedDisplay)
+	bool bInherit = false;
+};
+
+USTRUCT()
+struct VCAMEXTENSIONS_API FPerModifierClassWidgetSytleData
+{
+	GENERATED_BODY()
+
+	/** Data to attach to modifiers matching the associated class. */
+	UPROPERTY(EditAnywhere, Category = "Virtual Camera")
+	FWidgetStyleDataConfig ModifierStyles;
+
+	/** Meta data that is attached to a custom name, such as a connection point or group name (if using e.g. UClassBasedModifierHierarchyRules) */
+	UPROPERTY(EditAnywhere, Category = "Virtual Camera")
+	TMap<FName, FWidgetStyleDataConfig> CategorizedStyles;
+};
+
+/**
+ * Assigns meta data based on modifier class. To every modifier class you can assign meta data.
+ */
+UCLASS()
+class VCAMEXTENSIONS_API UClassBasedWidgetStyleDefinitions : public UModifierBoundWidgetStyleDefinitions
+{
+	GENERATED_BODY()
+public:
+
+	UPROPERTY(EditAnywhere, Category = "Virtual Camera")
+	TMap<TSubclassOf<UVCamModifier>, FPerModifierClassWidgetSytleData> Config;
+
+	/** Meta data that is not bound to any modifier but just to a simple name. */
+	UPROPERTY(EditAnywhere, Category = "Virtual Camera")
+	TMap<FName, FWidgetStyleDataArray> CategoriesWithoutModifier;
+	
+	//~ Begin UModifierMetaDataRules Interface
+	virtual TArray<UWidgetStyleData*> GetStylesForModifier_Implementation(UVCamComponent* VCamComponent, FName ModifierId) const override;
+	virtual TArray<UWidgetStyleData*> GetStylesForCategoryInModifier_Implementation(UVCamComponent* VCamComponent, FName ModifierId, FName ConnectionPointId) const override;
+	virtual TArray<UWidgetStyleData*> GetStylesForCategoryWithoutModifier_Implementation(FName Category) const override;
+	//~ End UModifierMetaDataRules Interface
+};
diff --git a/Engine/Plugins/Experimental/VirtualCamera/Source/VCamExtensions/Public/Styling/ModifierBoundWidgetStyleDefinitions.h b/Engine/Plugins/Experimental/VirtualCamera/Source/VCamExtensions/Public/Styling/ModifierBoundWidgetStyleDefinitions.h
new file mode 100644
--- /dev/null
+++ b/Engine/Plugins/Experimental/VirtualCamera/Source/VCamExtensions/Public/Styling/ModifierBoundWidgetStyleDefinitions.h
@@ -0,0 +1,44 @@
+// Copyright Epic Games, Inc. All Rights Reserved.
+
+#pragma once
+
+#include "WidgetStyleData.h"
+#include "Templates/SubclassOf.h"
+#include "UObject/Object.h"
+#include "ModifierBoundWidgetStyleDefinitions.generated.h"
+
+class UWidgetStyleData;
+class UVCamComponent;
+
+/**
+ * Associates information with modifiers and their connection points.
+ * You can use it to retrieve custom information assigned to a modifier and / or its connections,
+ * such as what icon a button representing that widget should have.
+ */
+UCLASS(Abstract, Blueprintable, EditInlineNew)
+class VCAMEXTENSIONS_API UModifierBoundWidgetStyleDefinitions : public UObject
+{
+	GENERATED_BODY()
+public:
+
+	/** Retrieves all meta data that is associated for a given modifier. */
+	UFUNCTION(BlueprintPure, BlueprintNativeEvent, Category = "Virtual Camera")
+	TArray<UWidgetStyleData*> GetStylesForModifier(UVCamComponent* VCamComponent, FName ModifierId) const;
+
+	/** Retrieves all meta data that is associated for a given modifier and a sub-category name. */
+	UFUNCTION(BlueprintPure, BlueprintNativeEvent, Category = "Virtual Camera")
+	TArray<UWidgetStyleData*> GetStylesForCategoryInModifier(UVCamComponent* VCamComponent, FName ModifierId, FName Category) const;
+
+	/** Retrieves all meta data that is associated with a given a category name; this data is not associated with any kind of modifier. */
+	UFUNCTION(BlueprintPure, BlueprintNativeEvent, Category = "Virtual Camera")
+	TArray<UWidgetStyleData*> GetStylesForCategoryWithoutModifier(FName Category) const;
+
+	UFUNCTION(BlueprintPure, Category = "Virtual Camera", meta = (DeterminesOutputType = "Class"))
+	UWidgetStyleData* GetStyleForModifierByClass(UVCamComponent* VCamComponent, FName ModifierId, TSubclassOf<UWidgetStyleData> Class) const;
+	
+	UFUNCTION(BlueprintPure, Category = "Virtual Camera", meta = (DeterminesOutputType = "Class"))
+	UWidgetStyleData* GetStyleForCategoryByClassInModifier(UVCamComponent* VCamComponent, FName ModifierId, FName Category, TSubclassOf<UWidgetStyleData> Class) const;
+	
+	UFUNCTION(BlueprintPure, Category = "Virtual Camera", meta = (DeterminesOutputType = "Class"))
+	UWidgetStyleData* GetStyleForCategoryWithoutModifier( FName Category, TSubclassOf<UWidgetStyleData> Class) const;
+};
diff --git a/Engine/Plugins/Experimental/VirtualCamera/Source/VCamExtensions/Public/Styling/ModifierBoundWidgetStylesAsset.h b/Engine/Plugins/Experimental/VirtualCamera/Source/VCamExtensions/Public/Styling/ModifierBoundWidgetStylesAsset.h
new file mode 100644
--- /dev/null
+++ b/Engine/Plugins/Experimental/VirtualCamera/Source/VCamExtensions/Public/Styling/ModifierBoundWidgetStylesAsset.h
@@ -0,0 +1,44 @@
+// Copyright Epic Games, Inc. All Rights Reserved.
+
+#pragma once
+
+#include "UObject/Object.h"
+#include "ModifierBoundWidgetStyleDefinitions.h"
+#include "ModifierBoundWidgetStylesAsset.generated.h"
+
+/**
+ * An asset intended to be referenced by Slate widgets.
+ * 
+ * For example, you can retrieve custom style info assigned to a modifier and / or its connections,
+ * such as what icon a button representing that widget should have.
+ */
+UCLASS(BlueprintType)
+class VCAMEXTENSIONS_API UModifierBoundWidgetStylesAsset : public UObject
+{
+	GENERATED_BODY()
+public:
+	
+	UPROPERTY(EditAnywhere, BlueprintReadOnly, Instanced, Category = "Virtual Camera")
+	TObjectPtr<UModifierBoundWidgetStyleDefinitions> Rules;
+
+	/** Retrieves all meta data that is associated for a given modifier. */
+	UFUNCTION(BlueprintPure, Category = "Virtual Camera")
+	TArray<UWidgetStyleData*> GetStylesForModifier(UVCamComponent* VCamComponent, FName ModifierId) const;
+
+	/** Retrieves all meta data that is associated for a given modifier and a sub-category name. */
+	UFUNCTION(BlueprintPure, Category = "Virtual Camera")
+	TArray<UWidgetStyleData*> GetStylesForCategoryInModifier(UVCamComponent* VCamComponent, FName ModifierId, FName Category) const;
+
+	/** Retrieves all meta data that is associated with a given a category name; this data is not associated with any kind of modifier. */
+	UFUNCTION(BlueprintPure, Category = "Virtual Camera")
+	TArray<UWidgetStyleData*> GetStylesForCategoryWithoutModifier(FName Category) const;
+
+	UFUNCTION(BlueprintPure, Category = "Virtual Camera", meta = (DeterminesOutputType = "Class"))
+	UWidgetStyleData* GetStyleForModifierByClass(UVCamComponent* VCamComponent, FName ModifierId, TSubclassOf<UWidgetStyleData> Class) const;
+	
+	UFUNCTION(BlueprintPure, Category = "Virtual Camera", meta = (DeterminesOutputType = "Class"))
+	UWidgetStyleData* GetStyleForCategoryByClassInModifier(UVCamComponent* VCamComponent, FName ModifierId, FName Category, TSubclassOf<UWidgetStyleData> Class) const;
+	
+	UFUNCTION(BlueprintPure, Category = "Virtual Camera", meta = (DeterminesOutputType = "Class"))
+	UWidgetStyleData* GetStyleForCategoryWithoutModifier( FName Category, TSubclassOf<UWidgetStyleData> Class) const;
+};
\ No newline at end of file
diff --git a/Engine/Plugins/Experimental/VirtualCamera/Source/VCamExtensions/Public/Styling/WidgetStyleData.h b/Engine/Plugins/Experimental/VirtualCamera/Source/VCamExtensions/Public/Styling/WidgetStyleData.h
new file mode 100644
--- /dev/null
+++ b/Engine/Plugins/Experimental/VirtualCamera/Source/VCamExtensions/Public/Styling/WidgetStyleData.h
@@ -0,0 +1,13 @@
+// Copyright Epic Games, Inc. All Rights Reserved.
+
+#pragma once
+
+#include "UObject/Object.h"
+#include "WidgetStyleData.generated.h"
+
+/** Base class for data that can be associated with modifiers / their connection points. */
+UCLASS(Abstract, Blueprintable, EditInlineNew)
+class VCAMEXTENSIONS_API UWidgetStyleData : public UObject
+{
+	GENERATED_BODY()
+};
\ No newline at end of file
diff --git a/Engine/Plugins/Experimental/VirtualCamera/Source/VCamExtensions/VCamExtensions.build.cs b/Engine/Plugins/Experimental/VirtualCamera/Source/VCamExtensions/VCamExtensions.build.cs
new file mode 100644
--- /dev/null
+++ b/Engine/Plugins/Experimental/VirtualCamera/Source/VCamExtensions/VCamExtensions.build.cs
@@ -0,0 +1,31 @@
+// Copyright Epic Games, Inc. All Rights Reserved.
+
+using UnrealBuildTool;
+
+public class VCamExtensions : ModuleRules
+{
+	public VCamExtensions(ReadOnlyTargetRules Target) : base(Target)
+	{
+		PCHUsage = ModuleRules.PCHUsageMode.UseExplicitOrSharedPCHs;
+		
+		PublicDependencyModuleNames.AddRange(
+			new string[]
+			{
+				"Core",
+			}
+		);
+			
+		
+		PrivateDependencyModuleNames.AddRange(
+			new string[]
+			{
+				"CoreUObject",
+				"Engine",
+				"InputCore",
+				"Slate",
+				"SlateCore",
+				"VCamCore",
+			}
+		);
+	}
+}
diff --git a/Engine/Plugins/Experimental/VirtualCamera/Source/VCamExtensionsEditor/Private/AssetTypeActions/AssetTypeActions_ModifierBoundWidgetStylesAsset.h b/Engine/Plugins/Experimental/VirtualCamera/Source/VCamExtensionsEditor/Private/AssetTypeActions/AssetTypeActions_ModifierBoundWidgetStylesAsset.h
new file mode 100644
--- /dev/null
+++ b/Engine/Plugins/Experimental/VirtualCamera/Source/VCamExtensionsEditor/Private/AssetTypeActions/AssetTypeActions_ModifierBoundWidgetStylesAsset.h
@@ -0,0 +1,24 @@
+// Copyright Epic Games, Inc. All Rights Reserved.
+
+#pragma once
+
+#include "AssetTypeActions_Base.h"
+#include "IVCamCoreEditorModule.h"
+#include "Styling/ModifierBoundWidgetStylesAsset.h"
+
+namespace UE::VCamExtensionsEditor::Private
+{
+	class FAssetTypeActions_ModifierBoundWidgetStylesAsset : public FAssetTypeActions_Base
+	{
+	public:
+
+		//~ Begin IAssetTypeActions Interface
+		virtual FText GetName() const override { return NSLOCTEXT("AssetTypeActions", "FAssetTypeActions_ModifierBoundWidgetStyleCollectionAsset", "Modifier Bound Widget Styles"); }
+		virtual FColor GetTypeColor() const override { return FColor(201, 29, 85); }
+		virtual UClass* GetSupportedClass() const override { return UModifierBoundWidgetStylesAsset::StaticClass(); }
+		virtual uint32 GetCategories() override { return VCamCoreEditor::IVCamCoreEditorModule::Get().GetAdvancedAssetCategoryForVCam(); }
+		//~ End IAssetTypeActions Interface
+	};
+}
+
+
diff --git a/Engine/Plugins/Experimental/VirtualCamera/Source/VCamExtensionsEditor/Private/AssetTypeActions/AssetTypeActions_ModifierHierarchyAsset.h b/Engine/Plugins/Experimental/VirtualCamera/Source/VCamExtensionsEditor/Private/AssetTypeActions/AssetTypeActions_ModifierHierarchyAsset.h
new file mode 100644
--- /dev/null
+++ b/Engine/Plugins/Experimental/VirtualCamera/Source/VCamExtensionsEditor/Private/AssetTypeActions/AssetTypeActions_ModifierHierarchyAsset.h
@@ -0,0 +1,24 @@
+// Copyright Epic Games, Inc. All Rights Reserved.
+
+#pragma once
+
+#include "AssetTypeActions_Base.h"
+#include "IVCamCoreEditorModule.h"
+#include "Hierarchies/ModifierHierarchyAsset.h"
+
+namespace UE::VCamExtensionsEditor::Private
+{
+	class FAssetTypeActions_ModifierHierarchyAsset : public FAssetTypeActions_Base
+	{
+	public:
+
+		//~ Begin IAssetTypeActions Interface
+		virtual FText GetName() const override { return NSLOCTEXT("AssetTypeActions", "AssetTypeActions_ModifierHierarchyAsset", "Modifier Hierarchy"); }
+		virtual FColor GetTypeColor() const override { return FColor(201, 29, 85); }
+		virtual UClass* GetSupportedClass() const override { return UModifierHierarchyAsset::StaticClass(); }
+		virtual uint32 GetCategories() override { return VCamCoreEditor::IVCamCoreEditorModule::Get().GetAdvancedAssetCategoryForVCam(); }
+		//~ End IAssetTypeActions Interface
+	};
+}
+
+
diff --git a/Engine/Plugins/Experimental/VirtualCamera/Source/VCamExtensionsEditor/Private/Factories/ModifierBoundWidgetStylesAssetFactory.cpp b/Engine/Plugins/Experimental/VirtualCamera/Source/VCamExtensionsEditor/Private/Factories/ModifierBoundWidgetStylesAssetFactory.cpp
new file mode 100644
--- /dev/null
+++ b/Engine/Plugins/Experimental/VirtualCamera/Source/VCamExtensionsEditor/Private/Factories/ModifierBoundWidgetStylesAssetFactory.cpp
@@ -0,0 +1,48 @@
+// Copyright Epic Games, Inc. All Rights Reserved.
+
+#include "ModifierBoundWidgetStylesAssetFactory.h"
+
+#include "IVCamCoreEditorModule.h"
+#include "Styling/ModifierBoundWidgetStylesAsset.h"
+
+#define LOCTEXT_NAMESPACE "UModifierMetaDataAssetFactory"
+
+UModifierBoundWidgetStylesAssetFactory::UModifierBoundWidgetStylesAssetFactory()
+{
+	SupportedClass = UModifierBoundWidgetStylesAsset::StaticClass();
+
+	bCreateNew = true;
+	bEditAfterNew = true;
+}
+
+FText UModifierBoundWidgetStylesAssetFactory::GetDisplayName() const
+{
+	return LOCTEXT("DisplayName", "Modifier Bound Widget Styles");
+}
+
+FText UModifierBoundWidgetStylesAssetFactory::GetToolTip() const
+{
+	return UModifierBoundWidgetStylesAsset::StaticClass()->GetToolTipText();
+}
+
+UObject* UModifierBoundWidgetStylesAssetFactory::FactoryCreateNew(UClass* Class, UObject* InParent, FName Name, EObjectFlags Flags, UObject* Context, FFeedbackContext* Warn)
+{
+	return NewObject<UModifierBoundWidgetStylesAsset>(InParent, Class, Name, Flags);
+}
+
+uint32 UModifierBoundWidgetStylesAssetFactory::GetMenuCategories() const
+{
+	return UE::VCamCoreEditor::IVCamCoreEditorModule::Get().GetAdvancedAssetCategoryForVCam();
+}
+
+const TArray<FText>& UModifierBoundWidgetStylesAssetFactory::GetMenuCategorySubMenus() const
+{
+	static const TArray<FText> SubMenus
+	{
+		LOCTEXT("UModifierMetaDataAssetFactory_SubMenu", "Extensions")
+	};
+
+	return SubMenus;
+}
+
+#undef LOCTEXT_NAMESPACE
diff --git a/Engine/Plugins/Experimental/VirtualCamera/Source/VCamExtensionsEditor/Private/Factories/ModifierBoundWidgetStylesAssetFactory.h b/Engine/Plugins/Experimental/VirtualCamera/Source/VCamExtensionsEditor/Private/Factories/ModifierBoundWidgetStylesAssetFactory.h
new file mode 100644
--- /dev/null
+++ b/Engine/Plugins/Experimental/VirtualCamera/Source/VCamExtensionsEditor/Private/Factories/ModifierBoundWidgetStylesAssetFactory.h
@@ -0,0 +1,26 @@
+// Copyright Epic Games, Inc. All Rights Reserved.
+
+#pragma once
+
+#include "Factories/Factory.h"
+#include "ModifierBoundWidgetStylesAssetFactory.generated.h"
+
+/**
+ * 
+ */
+UCLASS()
+class UModifierBoundWidgetStylesAssetFactory : public UFactory
+{
+	GENERATED_BODY()
+public:
+
+	UModifierBoundWidgetStylesAssetFactory();
+
+	//~ Begin UFactory Interface
+	virtual FText GetDisplayName() const override;
+	virtual FText GetToolTip() const override;
+	virtual UObject* FactoryCreateNew(UClass* Class, UObject* InParent, FName Name, EObjectFlags Flags, UObject* Context, FFeedbackContext* Warn) override;
+	virtual uint32 GetMenuCategories() const override;
+	virtual const TArray<FText>& GetMenuCategorySubMenus() const override;
+	//~ End UFactory Interface
+};
diff --git a/Engine/Plugins/Experimental/VirtualCamera/Source/VCamExtensionsEditor/Private/Factories/ModifierHierarchyAssetFactory.cpp b/Engine/Plugins/Experimental/VirtualCamera/Source/VCamExtensionsEditor/Private/Factories/ModifierHierarchyAssetFactory.cpp
new file mode 100644
--- /dev/null
+++ b/Engine/Plugins/Experimental/VirtualCamera/Source/VCamExtensionsEditor/Private/Factories/ModifierHierarchyAssetFactory.cpp
@@ -0,0 +1,48 @@
+// Copyright Epic Games, Inc. All Rights Reserved.
+
+#include "ModifierHierarchyAssetFactory.h"
+
+#include "IVCamCoreEditorModule.h"
+#include "Hierarchies/ModifierHierarchyAsset.h"
+
+#define LOCTEXT_NAMESPACE "UModifierHierarchyAssetFactory"
+
+UModifierHierarchyAssetFactory::UModifierHierarchyAssetFactory()
+{
+	SupportedClass = UModifierHierarchyAsset::StaticClass();
+
+	bCreateNew = true;
+	bEditAfterNew = true;
+}
+
+FText UModifierHierarchyAssetFactory::GetDisplayName() const
+{
+	return LOCTEXT("DisplayName", "Modifier Hierarchy");
+}
+
+FText UModifierHierarchyAssetFactory::GetToolTip() const
+{
+	return UModifierHierarchyAsset::StaticClass()->GetToolTipText();
+}
+
+UObject* UModifierHierarchyAssetFactory::FactoryCreateNew(UClass* Class, UObject* InParent, FName Name, EObjectFlags Flags, UObject* Context, FFeedbackContext* Warn)
+{
+	return NewObject<UModifierHierarchyAsset>(InParent, Class, Name, Flags);
+}
+
+uint32 UModifierHierarchyAssetFactory::GetMenuCategories() const
+{
+	return UE::VCamCoreEditor::IVCamCoreEditorModule::Get().GetAdvancedAssetCategoryForVCam();
+}
+
+const TArray<FText>& UModifierHierarchyAssetFactory::GetMenuCategorySubMenus() const
+{
+	static const TArray<FText> SubMenus
+	{
+		LOCTEXT("UModifierMetaDataAssetFactory_SubMenu", "Extensions")
+	};
+
+	return SubMenus;
+}
+
+#undef LOCTEXT_NAMESPACE
diff --git a/Engine/Plugins/Experimental/VirtualCamera/Source/VCamExtensionsEditor/Private/Factories/ModifierHierarchyAssetFactory.h b/Engine/Plugins/Experimental/VirtualCamera/Source/VCamExtensionsEditor/Private/Factories/ModifierHierarchyAssetFactory.h
new file mode 100644
--- /dev/null
+++ b/Engine/Plugins/Experimental/VirtualCamera/Source/VCamExtensionsEditor/Private/Factories/ModifierHierarchyAssetFactory.h
@@ -0,0 +1,26 @@
+// Copyright Epic Games, Inc. All Rights Reserved.
+
+#pragma once
+
+#include "Factories/Factory.h"
+#include "ModifierHierarchyAssetFactory.generated.h"
+
+/**
+ * 
+ */
+UCLASS()
+class UModifierHierarchyAssetFactory : public UFactory
+{
+	GENERATED_BODY()
+public:
+
+	UModifierHierarchyAssetFactory();
+
+	//~ Begin UFactory Interface
+	virtual FText GetDisplayName() const override;
+	virtual FText GetToolTip() const override;
+	virtual UObject* FactoryCreateNew(UClass* Class, UObject* InParent, FName Name, EObjectFlags Flags, UObject* Context, FFeedbackContext* Warn) override;
+	virtual uint32 GetMenuCategories() const override;
+	virtual const TArray<FText>& GetMenuCategorySubMenus() const override;
+	//~ End UFactory Interface
+};
diff --git a/Engine/Plugins/Experimental/VirtualCamera/Source/VCamExtensionsEditor/Private/VCamExtensionsEditor.cpp b/Engine/Plugins/Experimental/VirtualCamera/Source/VCamExtensionsEditor/Private/VCamExtensionsEditor.cpp
new file mode 100644
--- /dev/null
+++ b/Engine/Plugins/Experimental/VirtualCamera/Source/VCamExtensionsEditor/Private/VCamExtensionsEditor.cpp
@@ -0,0 +1,57 @@
+// Copyright Epic Games, Inc. All Rights Reserved.
+
+#include "AssetTypeActions/AssetTypeActions_ModifierHierarchyAsset.h"
+#include "AssetTypeActions/AssetTypeActions_ModifierBoundWidgetStylesAsset.h"
+
+#include "AssetToolsModule.h"
+#include "IAssetTools.h"
+#include "Modules/ModuleInterface.h"
+#include "Modules/ModuleManager.h"
+
+namespace UE::VCamExtensionsEditor::Private
+{
+	class FVCamExtensionsEditorModule : public IModuleInterface
+	{
+	public:
+		
+		virtual void StartupModule() override
+		{
+			SetupAssetTypeActions();
+		}
+		
+		virtual void ShutdownModule() override
+		{
+			CleanUpAssetTypeActions();
+		}
+
+	private:
+
+		TSharedPtr<FAssetTypeActions_ModifierBoundWidgetStylesAsset> MetaDataActions;
+		TSharedPtr<FAssetTypeActions_ModifierHierarchyAsset> HierarchyActions;
+		
+		void SetupAssetTypeActions()
+		{
+			MetaDataActions = MakeShared<FAssetTypeActions_ModifierBoundWidgetStylesAsset>();
+			HierarchyActions = MakeShared<FAssetTypeActions_ModifierHierarchyAsset>();
+			
+			IAssetTools& AssetTools = FModuleManager::LoadModuleChecked<FAssetToolsModule>("AssetTools").Get();
+			AssetTools.RegisterAssetTypeActions(MetaDataActions.ToSharedRef());
+			AssetTools.RegisterAssetTypeActions(HierarchyActions.ToSharedRef());
+		}
+		
+		void CleanUpAssetTypeActions()
+		{
+			if (IModuleInterface* AssetToolsPtr = FModuleManager::Get().GetModule("AssetTools"))
+			{
+				IAssetTools& AssetTools = static_cast<FAssetToolsModule*>(AssetToolsPtr)->Get();
+				AssetTools.UnregisterAssetTypeActions(MetaDataActions.ToSharedRef());
+				AssetTools.UnregisterAssetTypeActions(HierarchyActions.ToSharedRef());
+			}
+
+			MetaDataActions.Reset();
+			HierarchyActions.Reset();
+		}
+	};
+}
+
+IMPLEMENT_MODULE(UE::VCamExtensionsEditor::Private::FVCamExtensionsEditorModule, VCamExtensionsEditor);
\ No newline at end of file
diff --git a/Engine/Plugins/Experimental/VirtualCamera/Source/VCamExtensionsEditor/Private/VCamExtensionsEditorLog.cpp b/Engine/Plugins/Experimental/VirtualCamera/Source/VCamExtensionsEditor/Private/VCamExtensionsEditorLog.cpp
new file mode 100644
--- /dev/null
+++ b/Engine/Plugins/Experimental/VirtualCamera/Source/VCamExtensionsEditor/Private/VCamExtensionsEditorLog.cpp
@@ -0,0 +1,5 @@
+// Copyright Epic Games, Inc. All Rights Reserved.
+
+#include "VCamExtensionsEditorLog.h"
+
+DEFINE_LOG_CATEGORY(LogVCamEditorExtensions);
\ No newline at end of file
diff --git a/Engine/Plugins/Experimental/VirtualCamera/Source/VCamExtensionsEditor/Private/VCamExtensionsEditorLog.h b/Engine/Plugins/Experimental/VirtualCamera/Source/VCamExtensionsEditor/Private/VCamExtensionsEditorLog.h
new file mode 100644
--- /dev/null
+++ b/Engine/Plugins/Experimental/VirtualCamera/Source/VCamExtensionsEditor/Private/VCamExtensionsEditorLog.h
@@ -0,0 +1,7 @@
+// Copyright Epic Games, Inc. All Rights Reserved.
+
+#pragma once
+
+#include "Logging/LogMacros.h"
+
+DECLARE_LOG_CATEGORY_EXTERN(LogVCamEditorExtensions, Log, All);
\ No newline at end of file
diff --git a/Engine/Plugins/Experimental/VirtualCamera/Source/VCamExtensionsEditor/VCamExtensionsEditor.build.cs b/Engine/Plugins/Experimental/VirtualCamera/Source/VCamExtensionsEditor/VCamExtensionsEditor.build.cs
new file mode 100644
--- /dev/null
+++ b/Engine/Plugins/Experimental/VirtualCamera/Source/VCamExtensionsEditor/VCamExtensionsEditor.build.cs
@@ -0,0 +1,32 @@
+// Copyright Epic Games, Inc. All Rights Reserved.
+
+using UnrealBuildTool;
+
+public class VCamExtensionsEditor : ModuleRules
+{
+	public VCamExtensionsEditor(ReadOnlyTargetRules Target) : base(Target)
+	{
+		PCHUsage = ModuleRules.PCHUsageMode.UseExplicitOrSharedPCHs;
+		
+		PublicDependencyModuleNames.AddRange(
+			new string[]
+			{
+				"Core",
+			}
+		);
+			
+		
+		PrivateDependencyModuleNames.AddRange(
+			new string[]
+			{
+				"AssetTools",
+				"CoreUObject",
+				"Engine",
+				"UnrealEd",
+				"VCamCore",
+				"VCamCoreEditor",
+				"VCamExtensions"
+			}
+		);
+	}
+}
diff --git a/Engine/Plugins/Experimental/VirtualCamera/VirtualCamera.uplugin b/Engine/Plugins/Experimental/VirtualCamera/VirtualCamera.uplugin
--- a/Engine/Plugins/Experimental/VirtualCamera/VirtualCamera.uplugin
+++ b/Engine/Plugins/Experimental/VirtualCamera/VirtualCamera.uplugin
@@ -51,6 +51,22 @@
 				"TVOS"
 			]
 		},
+		{
+			"Name": "VCamExtensions",
+			"Type": "Runtime",
+			"LoadingPhase": "Default",
+			"PlatformDenyList": [
+				"TVOS"
+			]
+		},
+		{
+			"Name": "VCamExtensionsEditor",
+			"Type": "Editor",
+			"LoadingPhase": "Default",
+			"PlatformDenyList": [
+				"TVOS"
+			]
+		},
 		{
 			"Name": "VCamInput",
 			"Type": "Runtime",
