Add GetParentGroup to hierarchies

diff --git a/Engine/Plugins/Experimental/VirtualCamera/Source/VCamExtensions/Private/Hierarchies/ClassBasedModifierHierarchyRules.cpp b/Engine/Plugins/Experimental/VirtualCamera/Source/VCamExtensions/Private/Hierarchies/ClassBasedModifierHierarchyRules.cpp
--- a/Engine/Plugins/Experimental/VirtualCamera/Source/VCamExtensions/Private/Hierarchies/ClassBasedModifierHierarchyRules.cpp
+++ b/Engine/Plugins/Experimental/VirtualCamera/Source/VCamExtensions/Private/Hierarchies/ClassBasedModifierHierarchyRules.cpp
@@ -24,6 +24,22 @@ FName UClassBasedModifierHierarchyRules::GetRootGroup_Implementation() const
 		: NAME_None;
 }
 
+bool UClassBasedModifierHierarchyRules::GetParentGroup_Implementation(FName ChildGroup, FName& ParentGroup) const
+{
+	bool bFound = false;
+	ForEachGroup([ChildGroup, &ParentGroup, &bFound](UClassBasedModifierGroup& Group, UClassBasedModifierGroup* Parent)
+	{
+		if (Group.GroupName == ChildGroup)
+		{
+			ParentGroup = Parent->GroupName;
+			bFound = true;
+			return EBreakBehavior::Break;
+		}
+		return EBreakBehavior::Continue;
+	});
+	return bFound;
+}
+
 bool UClassBasedModifierHierarchyRules::GetGroupOfModifier_Implementation(UVCamModifier* Modifier, FName& Group) const
 {
 	if (!Modifier)
@@ -75,7 +91,7 @@ void UClassBasedModifierHierarchyRules::PostEditChangeProperty(FPropertyChangedE
 UClassBasedModifierGroup* UClassBasedModifierHierarchyRules::FindGroupByName(FName GroupName) const
 {
 	UClassBasedModifierGroup* FoundGroup = nullptr;
-	ForEachGroup([GroupName, &FoundGroup](UClassBasedModifierGroup& Group)
+	ForEachGroup([GroupName, &FoundGroup](UClassBasedModifierGroup& Group, UClassBasedModifierGroup* Parent)
 	{
 		if (Group.GroupName == GroupName)
 		{
@@ -118,7 +134,7 @@ TSet<UVCamModifier*> UClassBasedModifierHierarchyRules::EnumerateModifiersInGrou
 UClassBasedModifierGroup* UClassBasedModifierHierarchyRules::FindBestMatchFor(UVCamModifier& Modifier) const
 {
 	TArray<TPair<TSubclassOf<UVCamModifier>, UClassBasedModifierGroup&>> MatchingGroups;
-	ForEachGroup([&Modifier, &MatchingGroups](UClassBasedModifierGroup& Group)
+	ForEachGroup([&Modifier, &MatchingGroups](UClassBasedModifierGroup& Group, UClassBasedModifierGroup* Parent)
 	{
 		for (const TSubclassOf<UVCamModifier>& Class : Group.ModifierClasses)
 		{
@@ -155,30 +171,33 @@ UClassBasedModifierGroup* UClassBasedModifierHierarchyRules::FindBestMatchFor(UV
 	return ClosestGroup;
 }
 
-void UClassBasedModifierHierarchyRules::ForEachGroup(TFunctionRef<EBreakBehavior(UClassBasedModifierGroup& Group)> Callback) const
+void UClassBasedModifierHierarchyRules::ForEachGroup(TFunctionRef<EBreakBehavior(UClassBasedModifierGroup& CurrentGroup, UClassBasedModifierGroup* Parent)> Callback) const
 {
 	if (!ensure(RootGroup))
 	{
 		return;
 	}
 
-	TQueue<UClassBasedModifierGroup*> Queue;
-	Queue.Enqueue(RootGroup);
+	TQueue<TPair<UClassBasedModifierGroup*, UClassBasedModifierGroup*>> Queue;
+	Queue.Enqueue({ RootGroup, nullptr });
 
-	UClassBasedModifierGroup* Current = nullptr;
+	TPair<UClassBasedModifierGroup*, UClassBasedModifierGroup*> Current;
 	while (Queue.Dequeue(Current))
 	{
-		const EBreakBehavior BreakBehavior = Callback(*Current);
+		UClassBasedModifierGroup* CurrentNode = Current.Key;
+		UClassBasedModifierGroup* Parent = Current.Value;
+		
+		const EBreakBehavior BreakBehavior = Callback(*CurrentNode, Parent);
 		if (BreakBehavior == EBreakBehavior::Break)
 		{
 			return;
 		}
 		
-		for (UClassBasedModifierGroup* Child : Current->Children)
+		for (UClassBasedModifierGroup* Child : CurrentNode->Children)
 		{
 			if (Child)
 			{
-				Queue.Enqueue(Child);
+				Queue.Enqueue({ Child, CurrentNode });
 			}
 		}
 	}
diff --git a/Engine/Plugins/Experimental/VirtualCamera/Source/VCamExtensions/Private/Hierarchies/ModifierHierarchyRules.cpp b/Engine/Plugins/Experimental/VirtualCamera/Source/VCamExtensions/Private/Hierarchies/ModifierHierarchyRules.cpp
--- a/Engine/Plugins/Experimental/VirtualCamera/Source/VCamExtensions/Private/Hierarchies/ModifierHierarchyRules.cpp
+++ b/Engine/Plugins/Experimental/VirtualCamera/Source/VCamExtensions/Private/Hierarchies/ModifierHierarchyRules.cpp
@@ -8,6 +8,11 @@ FName UModifierHierarchyRules::GetRootGroup_Implementation() const
 	return {};
 }
 
+bool UModifierHierarchyRules::GetParentGroup_Implementation(FName ChildGroup, FName& ParentGroup) const
+{
+	return false;
+}
+
 bool UModifierHierarchyRules::GetGroupOfModifier_Implementation(UVCamModifier* Modifier, FName& Group) const
 {
 	unimplemented();
@@ -25,3 +30,4 @@ TSet<UVCamModifier*> UModifierHierarchyRules::GetModifiersInGroup_Implementation
 	unimplemented();
 	return {};
 }
+
diff --git a/Engine/Plugins/Experimental/VirtualCamera/Source/VCamExtensions/Public/Hierarchies/ClassBasedModifierHierarchyRules.h b/Engine/Plugins/Experimental/VirtualCamera/Source/VCamExtensions/Public/Hierarchies/ClassBasedModifierHierarchyRules.h
--- a/Engine/Plugins/Experimental/VirtualCamera/Source/VCamExtensions/Public/Hierarchies/ClassBasedModifierHierarchyRules.h
+++ b/Engine/Plugins/Experimental/VirtualCamera/Source/VCamExtensions/Public/Hierarchies/ClassBasedModifierHierarchyRules.h
@@ -40,6 +40,7 @@ public:
 
 	//~ Begin UModifierHierarchyRules Interface
 	virtual FName GetRootGroup_Implementation() const override;
+	virtual bool GetParentGroup_Implementation(FName ChildGroup, FName& ParentGroup) const override;
 	virtual bool GetGroupOfModifier_Implementation(UVCamModifier* Modifier, FName& Group) const override;
 	virtual TSet<FName> GetChildGroups_Implementation(FName ParentGroup) const override;
 	virtual TSet<UVCamModifier*> GetModifiersInGroup_Implementation(UVCamComponent* Component, FName GroupName) const override;
@@ -66,5 +67,5 @@ private:
 	TSet<UVCamModifier*> EnumerateModifiersInGroup(UClassBasedModifierGroup& Group, UVCamComponent& Component) const;
 	
 	UClassBasedModifierGroup* FindBestMatchFor(UVCamModifier& Modifier) const;
-	void ForEachGroup(TFunctionRef<EBreakBehavior(UClassBasedModifierGroup& Group)> Callback) const;
+	void ForEachGroup(TFunctionRef<EBreakBehavior(UClassBasedModifierGroup& CurrentGroup, UClassBasedModifierGroup* Parent)> Callback) const;
 };
diff --git a/Engine/Plugins/Experimental/VirtualCamera/Source/VCamExtensions/Public/Hierarchies/ModifierHierarchyRules.h b/Engine/Plugins/Experimental/VirtualCamera/Source/VCamExtensions/Public/Hierarchies/ModifierHierarchyRules.h
--- a/Engine/Plugins/Experimental/VirtualCamera/Source/VCamExtensions/Public/Hierarchies/ModifierHierarchyRules.h
+++ b/Engine/Plugins/Experimental/VirtualCamera/Source/VCamExtensions/Public/Hierarchies/ModifierHierarchyRules.h
@@ -24,21 +24,25 @@ class VCAMEXTENSIONS_API UModifierHierarchyRules : public UObject
 public:
 
 	/** Gets the root of the tree. */
-	UFUNCTION(BlueprintPure, BlueprintNativeEvent, Category = "Virtual Camera")
+	UFUNCTION(BlueprintPure, BlueprintNativeEvent, Category = "Virtual Camera|Hierarchies")
 	FName GetRootGroup() const;
 
+	/** Gets the parent of this given group. Fails if called on the root node. */
+	UFUNCTION(BlueprintPure, BlueprintNativeEvent, Category = "Virtual Camera|Hierarchies")
+	bool GetParentGroup(FName ChildGroup, FName& ParentGroup) const;
+
 	/**
 	 * Gets the group the modifier belongs to.
 	 * @return True if the modifier belongs to any group
 	 */
-	UFUNCTION(BlueprintPure, BlueprintNativeEvent, Category = "Virtual Camera")
+	UFUNCTION(BlueprintPure, BlueprintNativeEvent, Category = "Virtual Camera|Hierarchies")
 	bool GetGroupOfModifier(UVCamModifier* Modifier, FName& Group) const;
 
 	/** Gets the child groups of the given group. */
-	UFUNCTION(BlueprintPure, BlueprintNativeEvent, Category = "Virtual Camera")
+	UFUNCTION(BlueprintPure, BlueprintNativeEvent, Category = "Virtual Camera|Hierarchies")
 	TSet<FName> GetChildGroups(FName ParentGroup) const;
 
 	/** Gets all the modifiers on the component that belong in the given group. */
-	UFUNCTION(BlueprintPure, BlueprintNativeEvent, Category = "Virtual Camera")
+	UFUNCTION(BlueprintPure, BlueprintNativeEvent, Category = "Virtual Camera|Hierarchies")
 	TSet<UVCamModifier*> GetModifiersInGroup(UVCamComponent* Component, FName GroupName) const;
 };
