Add BaseModifierGroup which can later be used to implement shared logic for making sure GroupName is unique

diff --git a/Engine/Plugins/Experimental/VirtualCamera/Source/VCamExtensions/Private/Hierarchies/BaseModifierGroup.cpp b/Engine/Plugins/Experimental/VirtualCamera/Source/VCamExtensions/Private/Hierarchies/BaseModifierGroup.cpp
new file mode 100644
--- /dev/null
+++ b/Engine/Plugins/Experimental/VirtualCamera/Source/VCamExtensions/Private/Hierarchies/BaseModifierGroup.cpp
@@ -0,0 +1,21 @@
+// Copyright Epic Games, Inc. All Rights Reserved.
+
+#include "Hierarchies/BaseModifierGroup.h"
+
+void UBaseModifierGroup::PostInitProperties()
+{
+	GroupName = GetFName();
+	UObject::PostInitProperties();
+}
+
+#if WITH_EDITOR
+void UBaseModifierGroup::PostEditChangeProperty(FPropertyChangedEvent& PropertyChangedEvent)
+{
+	if (PropertyChangedEvent.Property && PropertyChangedEvent.Property->GetFName() == GET_MEMBER_NAME_CHECKED(UBaseModifierGroup, GroupName))
+	{
+		// TODO:
+	}
+	
+	UObject::PostEditChangeProperty(PropertyChangedEvent);
+}
+#endif
diff --git a/Engine/Plugins/Experimental/VirtualCamera/Source/VCamExtensions/Private/Hierarchies/ClassBasedModifierHierarchyRules.cpp b/Engine/Plugins/Experimental/VirtualCamera/Source/VCamExtensions/Private/Hierarchies/ClassBasedModifierHierarchyRules.cpp
--- a/Engine/Plugins/Experimental/VirtualCamera/Source/VCamExtensions/Private/Hierarchies/ClassBasedModifierHierarchyRules.cpp
+++ b/Engine/Plugins/Experimental/VirtualCamera/Source/VCamExtensions/Private/Hierarchies/ClassBasedModifierHierarchyRules.cpp
@@ -6,12 +6,6 @@
 
 #include "Algo/AllOf.h"
 
-void UClassBasedModifierGroup::PostInitProperties()
-{
-	GroupName = GetFName();
-	Super::PostInitProperties();
-}
-
 UClassBasedModifierHierarchyRules::UClassBasedModifierHierarchyRules()
 {
 	RootGroup = CreateDefaultSubobject<UClassBasedModifierGroup>(TEXT("Root"));
@@ -75,19 +69,6 @@ TSet<UVCamModifier*> UClassBasedModifierHierarchyRules::GetModifiersInGroup_Impl
 		: TSet<UVCamModifier*>{};
 }
 
-#if WITH_EDITOR
-void UClassBasedModifierHierarchyRules::PostEditChangeProperty(FPropertyChangedEvent& PropertyChangedEvent)
-{
-	if (PropertyChangedEvent.Property
-		&& PropertyChangedEvent.Property->GetFName() == GET_MEMBER_NAME_CHECKED(UClassBasedModifierGroup, GroupName))
-	{
-		// TODO: ensure unique name
-	}
-	
-	Super::PostEditChangeProperty(PropertyChangedEvent);
-}
-#endif
-
 UClassBasedModifierGroup* UClassBasedModifierHierarchyRules::FindGroupByName(FName GroupName) const
 {
 	UClassBasedModifierGroup* FoundGroup = nullptr;
diff --git a/Engine/Plugins/Experimental/VirtualCamera/Source/VCamExtensions/Public/Hierarchies/BaseModifierGroup.h b/Engine/Plugins/Experimental/VirtualCamera/Source/VCamExtensions/Public/Hierarchies/BaseModifierGroup.h
new file mode 100644
--- /dev/null
+++ b/Engine/Plugins/Experimental/VirtualCamera/Source/VCamExtensions/Public/Hierarchies/BaseModifierGroup.h
@@ -0,0 +1,27 @@
+// Copyright Epic Games, Inc. All Rights Reserved.
+
+#pragma once
+
+#include "UObject/Object.h"
+#include "BaseModifierGroup.generated.h"
+
+/** Implements base logic to keep group names unique within a hierachy */
+UCLASS(Abstract)
+class VCAMEXTENSIONS_API UBaseModifierGroup : public UObject
+{
+	GENERATED_BODY()
+public:
+
+	UPROPERTY(EditAnywhere, Category = "Virtual Camera|Hierarchies")
+	FName GroupName;
+
+	//~ Begin UObject Interface
+	virtual void PostInitProperties() override;
+	//~ End UObject Interface
+
+#if WITH_EDITOR
+	//~ Begin UObject Interface
+	virtual void PostEditChangeProperty(FPropertyChangedEvent& PropertyChangedEvent) override;
+	//~ End UObject Interface
+#endif
+};
diff --git a/Engine/Plugins/Experimental/VirtualCamera/Source/VCamExtensions/Public/Hierarchies/ClassBasedModifierHierarchyRules.h b/Engine/Plugins/Experimental/VirtualCamera/Source/VCamExtensions/Public/Hierarchies/ClassBasedModifierHierarchyRules.h
--- a/Engine/Plugins/Experimental/VirtualCamera/Source/VCamExtensions/Public/Hierarchies/ClassBasedModifierHierarchyRules.h
+++ b/Engine/Plugins/Experimental/VirtualCamera/Source/VCamExtensions/Public/Hierarchies/ClassBasedModifierHierarchyRules.h
@@ -2,6 +2,7 @@
 
 #pragma once
 
+#include "BaseModifierGroup.h"
 #include "ModifierHierarchyRules.h"
 #include "Modifier/VCamModifier.h"
 #include "Templates/SubclassOf.h"
@@ -10,21 +11,16 @@
 class UVCamComponent;
 
 UCLASS(EditInlineNew)
-class VCAMEXTENSIONS_API UClassBasedModifierGroup : public UObject
+class VCAMEXTENSIONS_API UClassBasedModifierGroup : public UBaseModifierGroup
 {
 	GENERATED_BODY()
 public:
-
-	UPROPERTY(EditAnywhere, Category = "Virtual Camera")
-	FName GroupName;
 	
-	UPROPERTY(EditAnywhere, Category = "Virtual Camera")
+	UPROPERTY(EditAnywhere, Category = "Virtual Camera|Hierarchies")
 	TSet<TSubclassOf<UVCamModifier>> ModifierClasses;
 
-	UPROPERTY(EditAnywhere, Instanced, Category = "Virtual Camera")
+	UPROPERTY(EditAnywhere, Instanced, Category = "Virtual Camera|Hierarchies")
 	TArray<TObjectPtr<UClassBasedModifierGroup>> Children;
-
-	virtual void PostInitProperties() override;
 };
 
 /**
@@ -46,12 +42,6 @@ public:
 	virtual TSet<UVCamModifier*> GetModifiersInGroup_Implementation(UVCamComponent* Component, FName GroupName) const override;
 	//~ End UModifierHierarchyRules Interface
 
-#if WITH_EDITOR
-	//~ Begin UObject Interface
-	virtual void PostEditChangeProperty(FPropertyChangedEvent& PropertyChangedEvent) override;
-	//~ End UObject Interface
-#endif
-
 private:
 	
 	UPROPERTY(EditAnywhere, Instanced, Category = "Virtual Camera", meta = (NoResetToDefault))
