Extract ForEachGroup to reusable tree node type-agnostic function

diff --git a/Engine/Plugins/Experimental/VirtualCamera/Source/VCamExtensions/Private/Hierarchies/ClassBasedModifierHierarchyRules.cpp b/Engine/Plugins/Experimental/VirtualCamera/Source/VCamExtensions/Private/Hierarchies/ClassBasedModifierHierarchyRules.cpp
--- a/Engine/Plugins/Experimental/VirtualCamera/Source/VCamExtensions/Private/Hierarchies/ClassBasedModifierHierarchyRules.cpp
+++ b/Engine/Plugins/Experimental/VirtualCamera/Source/VCamExtensions/Private/Hierarchies/ClassBasedModifierHierarchyRules.cpp
@@ -2,6 +2,7 @@
 
 #include "Hierarchies/ClassBasedModifierHierarchyRules.h"
 
+#include "HierarchyUtils.h"
 #include "VCamComponent.h"
 
 #include "Algo/AllOf.h"
@@ -159,27 +160,13 @@ void UClassBasedModifierHierarchyRules::ForEachGroup(TFunctionRef<EBreakBehavior
 		return;
 	}
 
-	TQueue<TPair<UClassBasedModifierGroup*, UClassBasedModifierGroup*>> Queue;
-	Queue.Enqueue({ RootGroup, nullptr });
-
-	TPair<UClassBasedModifierGroup*, UClassBasedModifierGroup*> Current;
-	while (Queue.Dequeue(Current))
-	{
-		UClassBasedModifierGroup* CurrentNode = Current.Key;
-		UClassBasedModifierGroup* Parent = Current.Value;
-		
-		const EBreakBehavior BreakBehavior = Callback(*CurrentNode, Parent);
-		if (BreakBehavior == EBreakBehavior::Break)
-		{
-			return;
-		}
-		
-		for (UClassBasedModifierGroup* Child : CurrentNode->Children)
+	using namespace UE::VCamExtensions;
+	HierarchyUtils::ForEachGroup(
+		*RootGroup,
+		[Callback](UClassBasedModifierGroup& CurrentGroup, UClassBasedModifierGroup* Parent)
 		{
-			if (Child)
-			{
-				Queue.Enqueue({ Child, CurrentNode });
-			}
-		}
-	}
+			return Callback(CurrentGroup, Parent) == EBreakBehavior::Continue ? HierarchyUtils::EBreakBehavior::Continue : HierarchyUtils::EBreakBehavior::Break;
+		},
+		[](UClassBasedModifierGroup& CurrentGroup){ return CurrentGroup.Children; }
+		);
 }
diff --git a/Engine/Plugins/Experimental/VirtualCamera/Source/VCamExtensions/Private/Hierarchies/HierarchyUtils.h b/Engine/Plugins/Experimental/VirtualCamera/Source/VCamExtensions/Private/Hierarchies/HierarchyUtils.h
new file mode 100644
--- /dev/null
+++ b/Engine/Plugins/Experimental/VirtualCamera/Source/VCamExtensions/Private/Hierarchies/HierarchyUtils.h
@@ -0,0 +1,50 @@
+// Copyright Epic Games, Inc. All Rights Reserved.
+
+#pragma once
+
+#include "Containers/Map.h"
+#include "Containers/Queue.h"
+
+namespace UE::VCamExtensions::HierarchyUtils
+{
+	enum class EBreakBehavior
+	{
+		Continue,
+		Break
+	};
+
+	
+	/**
+	 * Utility function for iterating a tree like structure.
+	 *
+	 * ProcessRelationShipFunc = EBreakBehavior(TNodeType& CurrentGroup, TNodeType* Parent)
+	 * GetChildrenFunc = TArray<TNodeType*>(TNodeType& CurrentGroup)
+	 */
+	template<typename TNodeType, typename TProcessRelationShipFunc, typename TGetChildrenFunc>
+	void ForEachGroup(TNodeType& RootNode, TProcessRelationShipFunc ProcessRelationShipCallback, TGetChildrenFunc GetChildrenFunc)
+	{
+		TQueue<TPair<TNodeType*, TNodeType*>> Queue;
+		Queue.Enqueue({ &RootNode, nullptr });
+
+		TPair<TNodeType*, TNodeType*> Current;
+		while (Queue.Dequeue(Current))
+		{
+			TNodeType* CurrentNode = Current.Key;
+			TNodeType* Parent = Current.Value;
+		
+			const HierarchyUtils::EBreakBehavior BreakBehavior = ProcessRelationShipCallback(*CurrentNode, Parent);
+			if (BreakBehavior == EBreakBehavior::Break)
+			{
+				return;
+			}
+		
+			for (TNodeType* Child : GetChildrenFunc(*CurrentNode))
+			{
+				if (Child)
+				{
+					Queue.Enqueue({ Child, CurrentNode });
+				}
+			}
+		}
+	}
+};
diff --git a/Engine/Plugins/Experimental/VirtualCamera/Source/VCamExtensions/Public/Hierarchies/ClassBasedModifierHierarchyRules.h b/Engine/Plugins/Experimental/VirtualCamera/Source/VCamExtensions/Public/Hierarchies/ClassBasedModifierHierarchyRules.h
--- a/Engine/Plugins/Experimental/VirtualCamera/Source/VCamExtensions/Public/Hierarchies/ClassBasedModifierHierarchyRules.h
+++ b/Engine/Plugins/Experimental/VirtualCamera/Source/VCamExtensions/Public/Hierarchies/ClassBasedModifierHierarchyRules.h
@@ -5,6 +5,7 @@
 #include "BaseModifierGroup.h"
 #include "ModifierHierarchyRules.h"
 #include "Modifier/VCamModifier.h"
+#include "Templates/Function.h"
 #include "Templates/SubclassOf.h"
 #include "ClassBasedModifierHierarchyRules.generated.h"
 
