Add alternate UHierarchyRules implementation which has one button per group

diff --git a/Engine/Plugins/Experimental/VirtualCamera/Source/VCamExtensions/Private/Hierarchies/ModifierHierarchyRules.cpp b/Engine/Plugins/Experimental/VirtualCamera/Source/VCamExtensions/Private/Hierarchies/ModifierHierarchyRules.cpp
--- a/Engine/Plugins/Experimental/VirtualCamera/Source/VCamExtensions/Private/Hierarchies/ModifierHierarchyRules.cpp
+++ b/Engine/Plugins/Experimental/VirtualCamera/Source/VCamExtensions/Private/Hierarchies/ModifierHierarchyRules.cpp
@@ -2,6 +2,8 @@
 
 #include "Hierarchies/ModifierHierarchyRules.h"
 
+#include "VCamExtensionsLog.h"
+
 FName UModifierHierarchyRules::GetRootGroup_Implementation() const
 {
 	unimplemented();
@@ -19,6 +21,12 @@ bool UModifierHierarchyRules::GetGroupOfModifier_Implementation(UVCamModifier* M
 	return false;
 }
 
+bool UModifierHierarchyRules::GetConnectionPointForModifier_Implementation(UVCamModifier* Modifier, FName& ConnectionPoint) const
+{
+	UE_LOG(LogVCamExtensions, Warning, TEXT("%s does not support GetConnectionPointForModifier"), *GetClass()->GetName());
+	return false;
+}
+
 TSet<FName> UModifierHierarchyRules::GetChildGroups_Implementation(FName ParentGroup) const
 {
 	unimplemented();
@@ -30,4 +38,3 @@ TSet<UVCamModifier*> UModifierHierarchyRules::GetModifiersInGroup_Implementation
 	unimplemented();
 	return {};
 }
-
diff --git a/Engine/Plugins/Experimental/VirtualCamera/Source/VCamExtensions/Private/Hierarchies/TargetModifierPerGroupHierarchyRules.cpp b/Engine/Plugins/Experimental/VirtualCamera/Source/VCamExtensions/Private/Hierarchies/TargetModifierPerGroupHierarchyRules.cpp
new file mode 100644
--- /dev/null
+++ b/Engine/Plugins/Experimental/VirtualCamera/Source/VCamExtensions/Private/Hierarchies/TargetModifierPerGroupHierarchyRules.cpp
@@ -0,0 +1,131 @@
+// Copyright Epic Games, Inc. All Rights Reserved.
+
+#include "Hierarchies/TargetModifierPerGroupHierarchyRules.h"
+
+#include "HierarchyUtils.h"
+#include "VCamComponent.h"
+#include "VCamExtensionsLog.h"
+
+UTargetModifierPerGroupHierarchyRules::UTargetModifierPerGroupHierarchyRules()
+{
+	Root = CreateDefaultSubobject<USingleModifierPerGroupWithTargetSettings>(TEXT("Root"));
+}
+
+FName UTargetModifierPerGroupHierarchyRules::GetRootGroup_Implementation() const
+{
+	return ensure(Root) ? Root->GroupName : NAME_None;
+}
+
+bool UTargetModifierPerGroupHierarchyRules::GetParentGroup_Implementation(FName ChildGroup, FName& ParentGroup) const
+{
+	bool bFound = false;
+	ForEachGroup([ChildGroup, &ParentGroup, &bFound](USingleModifierPerGroupWithTargetSettings& Group, USingleModifierPerGroupWithTargetSettings* Parent)
+	{
+		if (Group.GroupName == ChildGroup)
+		{
+			ParentGroup = Parent->GroupName;
+			bFound = true;
+			return EBreakBehavior::Break;
+		}
+		return EBreakBehavior::Continue;
+	});
+	return bFound;
+}
+
+bool UTargetModifierPerGroupHierarchyRules::GetGroupOfModifier_Implementation(UVCamModifier* Modifier, FName& Group) const
+{
+	if (const USingleModifierPerGroupWithTargetSettings* FoundGroup = GetNodeForModifier(Modifier))
+	{
+		Group = FoundGroup->GroupName;
+		return true;
+	}
+	return false;
+}
+
+bool UTargetModifierPerGroupHierarchyRules::GetConnectionPointForModifier_Implementation(UVCamModifier* Modifier, FName& ConnectionPoint) const
+{
+	const USingleModifierPerGroupWithTargetSettings* FoundGroup = GetNodeForModifier(Modifier);
+	if (!FoundGroup || !Modifier->ConnectionPoints.Contains(FoundGroup->TargetSettings.TargetConnectionPoint))
+	{
+		return false;
+	}
+
+	ConnectionPoint = FoundGroup->TargetSettings.TargetConnectionPoint;
+	return true;
+}
+
+TSet<FName> UTargetModifierPerGroupHierarchyRules::GetChildGroups_Implementation(FName ParentGroup) const
+{
+	if (const USingleModifierPerGroupWithTargetSettings* Group = FindGroupByName(ParentGroup))
+	{
+		TSet<FName> Result;
+		Algo::TransformIf(Group->ChildElements, Result, [](const TObjectPtr<USingleModifierPerGroupWithTargetSettings>& Child){ return Child != nullptr; }, [](const TObjectPtr<USingleModifierPerGroupWithTargetSettings>& Child){ return Child->GroupName; });
+		return Result;
+	}
+	return {};
+}
+
+TSet<UVCamModifier*> UTargetModifierPerGroupHierarchyRules::GetModifiersInGroup_Implementation(UVCamComponent* Component, FName GroupName) const
+{
+	if (const USingleModifierPerGroupWithTargetSettings* Group = FindGroupByName(GroupName))
+	{
+		UVCamModifier* Modifier = Component->GetModifierByName(Group->TargetSettings.TargetModifierName);
+		return Modifier ? TSet<UVCamModifier*>{ Modifier } : TSet<UVCamModifier*>{};
+	}
+
+	return {};
+}
+
+USingleModifierPerGroupWithTargetSettings* UTargetModifierPerGroupHierarchyRules::GetNodeForModifier(UVCamModifier* Modifier) const
+{
+	if (!Modifier)
+	{
+		return nullptr;
+	}
+
+	USingleModifierPerGroupWithTargetSettings* FoundGroup = nullptr;
+	ForEachGroup([&FoundGroup, SearchName = Modifier->GetStackEntryName()](USingleModifierPerGroupWithTargetSettings& CurrentGroup, USingleModifierPerGroupWithTargetSettings* Parent)
+	{
+		if (CurrentGroup.TargetSettings.TargetModifierName == SearchName)
+		{
+			FoundGroup = &CurrentGroup;
+			return EBreakBehavior::Break; 
+		}
+		return EBreakBehavior::Continue;
+	});
+	return FoundGroup;
+}
+
+USingleModifierPerGroupWithTargetSettings* UTargetModifierPerGroupHierarchyRules::FindGroupByName(FName GroupName) const
+{
+	USingleModifierPerGroupWithTargetSettings* FoundGroup = nullptr;
+	ForEachGroup([GroupName, &FoundGroup](USingleModifierPerGroupWithTargetSettings& Group, USingleModifierPerGroupWithTargetSettings* Parent)
+	{
+		if (Group.GroupName == GroupName)
+		{
+			FoundGroup = &Group;
+			return EBreakBehavior::Break;
+		}
+		return EBreakBehavior::Continue;
+	});
+	UE_CLOG(FoundGroup == nullptr, LogVCamExtensions, Warning, TEXT("GroupName %s not found (%s)"), *GroupName.ToString(), *GetPathName());
+	return FoundGroup;
+}
+
+void UTargetModifierPerGroupHierarchyRules::ForEachGroup(TFunctionRef<EBreakBehavior(USingleModifierPerGroupWithTargetSettings& CurrentGroup, USingleModifierPerGroupWithTargetSettings* Parent)> Callback) const
+{
+	if (!ensure(Root))
+	{
+		return;
+	}
+
+	using namespace UE::VCamExtensions;
+	HierarchyUtils::ForEachGroup(
+		*Root,
+		[Callback](USingleModifierPerGroupWithTargetSettings& CurrentGroup, USingleModifierPerGroupWithTargetSettings* Parent)
+		{
+			return Callback(CurrentGroup, Parent) == EBreakBehavior::Continue ? HierarchyUtils::EBreakBehavior::Continue : HierarchyUtils::EBreakBehavior::Break;
+		},
+		[](USingleModifierPerGroupWithTargetSettings& CurrentGroup){ return CurrentGroup.ChildElements; }
+		);
+}
diff --git a/Engine/Plugins/Experimental/VirtualCamera/Source/VCamExtensions/Public/Hierarchies/BaseModifierGroup.h b/Engine/Plugins/Experimental/VirtualCamera/Source/VCamExtensions/Public/Hierarchies/BaseModifierGroup.h
--- a/Engine/Plugins/Experimental/VirtualCamera/Source/VCamExtensions/Public/Hierarchies/BaseModifierGroup.h
+++ b/Engine/Plugins/Experimental/VirtualCamera/Source/VCamExtensions/Public/Hierarchies/BaseModifierGroup.h
@@ -12,7 +12,7 @@ class VCAMEXTENSIONS_API UBaseModifierGroup : public UObject
 	GENERATED_BODY()
 public:
 
-	UPROPERTY(EditAnywhere, Category = "Virtual Camera|Hierarchies")
+	UPROPERTY(EditAnywhere, Category = "Hierarchy")
 	FName GroupName;
 
 	//~ Begin UObject Interface
diff --git a/Engine/Plugins/Experimental/VirtualCamera/Source/VCamExtensions/Public/Hierarchies/ClassBasedModifierHierarchyRules.h b/Engine/Plugins/Experimental/VirtualCamera/Source/VCamExtensions/Public/Hierarchies/ClassBasedModifierHierarchyRules.h
--- a/Engine/Plugins/Experimental/VirtualCamera/Source/VCamExtensions/Public/Hierarchies/ClassBasedModifierHierarchyRules.h
+++ b/Engine/Plugins/Experimental/VirtualCamera/Source/VCamExtensions/Public/Hierarchies/ClassBasedModifierHierarchyRules.h
@@ -17,10 +17,10 @@ class VCAMEXTENSIONS_API UClassBasedModifierGroup : public UBaseModifierGroup
 	GENERATED_BODY()
 public:
 	
-	UPROPERTY(EditAnywhere, Category = "Virtual Camera|Hierarchies")
+	UPROPERTY(EditAnywhere, Category = "Hierarchy")
 	TSet<TSubclassOf<UVCamModifier>> ModifierClasses;
 
-	UPROPERTY(EditAnywhere, Instanced, Category = "Virtual Camera|Hierarchies")
+	UPROPERTY(EditAnywhere, Instanced, Category = "Hierarchy")
 	TArray<TObjectPtr<UClassBasedModifierGroup>> Children;
 };
 
diff --git a/Engine/Plugins/Experimental/VirtualCamera/Source/VCamExtensions/Public/Hierarchies/ModifierHierarchyRules.h b/Engine/Plugins/Experimental/VirtualCamera/Source/VCamExtensions/Public/Hierarchies/ModifierHierarchyRules.h
--- a/Engine/Plugins/Experimental/VirtualCamera/Source/VCamExtensions/Public/Hierarchies/ModifierHierarchyRules.h
+++ b/Engine/Plugins/Experimental/VirtualCamera/Source/VCamExtensions/Public/Hierarchies/ModifierHierarchyRules.h
@@ -11,6 +11,7 @@ class UVCamModifier;
 /**
  * Defines a tree hierarchy. Each node is called a group.
  * A group consists of modifiers and (sub) groups.
+ * By design, modifiers are restricted to only belong to at most one group.
  * 
  * An example use case is if you want to have a button menu which should procedurally generate sub-button menus depending
  * on the modifiers in the component. 
@@ -31,13 +32,22 @@ public:
 	UFUNCTION(BlueprintPure, BlueprintNativeEvent, Category = "Virtual Camera|Hierarchies")
 	bool GetParentGroup(FName ChildGroup, FName& ParentGroup) const;
 
+	
 	/**
 	 * Gets the group the modifier belongs to.
 	 * @return True if the modifier belongs to any group
 	 */
 	UFUNCTION(BlueprintPure, BlueprintNativeEvent, Category = "Virtual Camera|Hierarchies")
 	bool GetGroupOfModifier(UVCamModifier* Modifier, FName& Group) const;
+	
+	/**
+	 * Gets the connection point the modifier is configured to be bound to.
+	 * This function is optional to implement; it is valid for it to always return false.
+	 */
+	UFUNCTION(BlueprintPure, BlueprintNativeEvent, Category = "Virtual Camera|Hierarchies")
+	bool GetConnectionPointForModifier(UVCamModifier* Modifier, FName& ConnectionPoint) const;
 
+	
 	/** Gets the child groups of the given group. */
 	UFUNCTION(BlueprintPure, BlueprintNativeEvent, Category = "Virtual Camera|Hierarchies")
 	TSet<FName> GetChildGroups(FName ParentGroup) const;
diff --git a/Engine/Plugins/Experimental/VirtualCamera/Source/VCamExtensions/Public/Hierarchies/TargetModifierPerGroupHierarchyRules.h b/Engine/Plugins/Experimental/VirtualCamera/Source/VCamExtensions/Public/Hierarchies/TargetModifierPerGroupHierarchyRules.h
new file mode 100644
--- /dev/null
+++ b/Engine/Plugins/Experimental/VirtualCamera/Source/VCamExtensions/Public/Hierarchies/TargetModifierPerGroupHierarchyRules.h
@@ -0,0 +1,58 @@
+// Copyright Epic Games, Inc. All Rights Reserved.
+
+#pragma once
+
+#include "BaseModifierGroup.h"
+#include "ModifierHierarchyRules.h"
+#include "Templates/Function.h"
+#include "UI/VCamConnectionStructs.h"
+#include "TargetModifierPerGroupHierarchyRules.generated.h"
+
+UCLASS(EditInlineNew)
+class VCAMEXTENSIONS_API USingleModifierPerGroupWithTargetSettings : public UBaseModifierGroup
+{
+	GENERATED_BODY()
+public:
+
+	UPROPERTY(EditAnywhere, Category = "Hierarchy")
+	FVCamConnectionTargetSettings TargetSettings;
+
+	UPROPERTY(EditAnywhere, Instanced, Category = "Hierarchy")
+	TArray<TObjectPtr<USingleModifierPerGroupWithTargetSettings>> ChildElements;
+};
+
+/**
+ * 
+ */
+UCLASS()
+class VCAMEXTENSIONS_API UTargetModifierPerGroupHierarchyRules : public UModifierHierarchyRules
+{
+	GENERATED_BODY()
+public:
+	
+	UTargetModifierPerGroupHierarchyRules();
+
+	//~ Begin UModifierHierarchyRules Interface
+	virtual FName GetRootGroup_Implementation() const override;
+	virtual bool GetParentGroup_Implementation(FName ChildGroup, FName& ParentGroup) const override;
+	virtual bool GetGroupOfModifier_Implementation(UVCamModifier* Modifier, FName& Group) const override;
+	virtual bool GetConnectionPointForModifier_Implementation(UVCamModifier* Modifier, FName& ConnectionPoint) const override;
+	virtual TSet<FName> GetChildGroups_Implementation(FName ParentGroup) const override;
+	virtual TSet<UVCamModifier*> GetModifiersInGroup_Implementation(UVCamComponent* Component, FName GroupName) const override;
+	//~ End UModifierHierarchyRules Interface
+
+private:
+
+	UPROPERTY(EditAnywhere, Instanced, Category = "Hierarchy")
+	TObjectPtr<USingleModifierPerGroupWithTargetSettings> Root;
+
+	USingleModifierPerGroupWithTargetSettings* GetNodeForModifier(UVCamModifier* Modifier) const;
+	USingleModifierPerGroupWithTargetSettings* FindGroupByName(FName GroupName) const;
+	
+	enum class EBreakBehavior
+	{
+		Continue,
+		Break
+	};
+	void ForEachGroup(TFunctionRef<EBreakBehavior(USingleModifierPerGroupWithTargetSettings& CurrentGroup, USingleModifierPerGroupWithTargetSettings* Parent)> Callback) const;
+};
