Fix non-unity error

diff --git a/Engine/Plugins/Experimental/VirtualCamera/Source/VCamExtensions/Public/Styling/ButtonWidgetStyleData.h b/Engine/Plugins/Experimental/VirtualCamera/Source/VCamExtensions/Public/Styling/ButtonWidgetStyleData.h
--- a/Engine/Plugins/Experimental/VirtualCamera/Source/VCamExtensions/Public/Styling/ButtonWidgetStyleData.h
+++ b/Engine/Plugins/Experimental/VirtualCamera/Source/VCamExtensions/Public/Styling/ButtonWidgetStyleData.h
@@ -3,6 +3,7 @@
 #pragma once
 
 #include "WidgetStyleData.h"
+#include "Styling/SlateTypes.h"
 #include "ButtonWidgetStyleData.generated.h"
 
 /** Associates a button style for widgets that want to display a button for a modifier / connection point. */
