Move files to LevelSequence folder

diff --git a/Engine/Plugins/Experimental/VirtualCamera/Source/VirtualCamera/Private/LevelSequence/LevelSequencePlaybackController.cpp b/Engine/Plugins/Experimental/VirtualCamera/Source/VirtualCamera/Private/LevelSequence/LevelSequencePlaybackController.cpp
new file mode 100644
--- /dev/null
+++ b/Engine/Plugins/Experimental/VirtualCamera/Source/VirtualCamera/Private/LevelSequence/LevelSequencePlaybackController.cpp
@@ -0,0 +1,345 @@
+// Copyright Epic Games, Inc. All Rights Reserved.
+
+#include "LevelSequence/LevelSequencePlaybackController.h"
+
+#include "AssetRegistry/AssetData.h"
+#include "AssetRegistry/AssetRegistryModule.h"
+#include "AssetRegistry/IAssetRegistry.h"
+#include "HAL/FileManager.h"
+#include "LevelSequence.h"
+#include "Modules/ModuleManager.h"
+#include "MovieScene.h"
+#include "MovieSceneTimeHelpers.h"
+
+#if WITH_EDITOR
+#include "ILevelSequenceEditorToolkit.h"
+#include "Input/Reply.h"
+#include "Recorder/TakeRecorderBlueprintLibrary.h"
+#include "Templates/SharedPointer.h"
+#include "Subsystems/AssetEditorSubsystem.h"
+#endif //WITH_EDITOR
+
+ULevelSequencePlaybackController::ULevelSequencePlaybackController(const FObjectInitializer& ObjectInitializer)
+	: Super(ObjectInitializer)
+	, ActiveLevelSequence(nullptr)
+{
+
+}
+
+void ULevelSequencePlaybackController::GetLevelSequences(TArray<FLevelSequenceData>& OutLevelSequenceNames)
+{
+	OutLevelSequenceNames.Empty();
+
+	TArray<FAssetData> LevelSequences;
+	LevelSequences.Empty();
+
+	FAssetRegistryModule* AssetRegistryModule = FModuleManager::Get().GetModulePtr<FAssetRegistryModule>("AssetRegistry");
+	if (AssetRegistryModule)
+	{
+		IAssetRegistry& AssetRegistry = AssetRegistryModule->Get();
+		if (!AssetRegistry.GetAssetsByClass(FTopLevelAssetPath(TEXT("/Script/LevelSequence"), TEXT("LevelSequence")), LevelSequences, false))
+		{
+			UE_LOG(LogActor, Error, TEXT("VirtualCamera - No Asset Registry module found!"));
+			return;
+		}
+	}
+
+	IFileManager& FileManager = IFileManager::Get();
+	for (FAssetData LevelSequence : LevelSequences)
+	{
+		// Get the file system name of the package so we can get other data on it (i.e. timestamp)
+		FString LevelSequenceFile = FPackageName::LongPackageNameToFilename(
+			LevelSequence.PackageName.ToString(), 
+			FPackageName::GetAssetPackageExtension()
+		);
+
+		FLevelSequenceData NewSequenceData = FLevelSequenceData(
+			LevelSequence.GetObjectPathString(),
+			LevelSequence.AssetName.ToString(),
+			FileManager.GetTimeStamp(*LevelSequenceFile)
+		);
+
+		OutLevelSequenceNames.Add(NewSequenceData);
+	}
+
+	// Sort the level sequence names alphabetically by Display Name
+	OutLevelSequenceNames.Sort([](const FLevelSequenceData& LeftItem, const FLevelSequenceData& RightItem)
+	{
+		return LeftItem.DisplayName < RightItem.DisplayName;
+	});
+}
+
+FString ULevelSequencePlaybackController::GetActiveLevelSequenceName() const 
+{
+	if (ActiveLevelSequence)
+	{
+		return ActiveLevelSequence->GetName();
+	}
+
+	return FString();
+}
+
+FFrameRate ULevelSequencePlaybackController::GetCurrentSequenceFrameRate() const
+{
+	if (ActiveLevelSequence)
+	{
+		return ActiveLevelSequence->GetMovieScene()->GetDisplayRate();
+	}
+
+	return FFrameRate();
+}
+
+bool ULevelSequencePlaybackController::IsSequencerLockedToCameraCut() const
+{
+#if WITH_EDITOR
+	if (ActiveLevelSequence)
+	{
+		TSharedPtr<ISequencer> Sequencer = WeakSequencer.Pin();
+		if (Sequencer)
+		{
+			return Sequencer->IsPerspectiveViewportCameraCutEnabled();
+		}
+	}
+#endif //WITH_EDITOR
+
+	return false;
+}
+
+void ULevelSequencePlaybackController::SetSequencerLockedToCameraCut(bool bLockView)
+{
+#if WITH_EDITOR
+	if (ActiveLevelSequence)
+	{
+		TSharedPtr<ISequencer> Sequencer = WeakSequencer.Pin();
+		if (Sequencer)
+		{
+			return Sequencer->SetPerspectiveViewportCameraCutEnabled(bLockView);
+		}
+	}
+#endif //WITH_EDITOR
+}
+
+FFrameNumber ULevelSequencePlaybackController::GetCurrentSequencePlaybackStart() const
+{
+	if (ActiveLevelSequence)
+	{
+		FFrameNumber Value = ActiveLevelSequence->GetMovieScene()->GetPlaybackRange().GetLowerBoundValue();
+		if (ActiveLevelSequence->GetMovieScene()->GetDisplayRate() != ActiveLevelSequence->GetMovieScene()->GetTickResolution())
+		{
+			const FFrameTime ConvertedTime = FFrameRate::TransformTime(FFrameTime(Value), ActiveLevelSequence->GetMovieScene()->GetTickResolution(), ActiveLevelSequence->GetMovieScene()->GetDisplayRate());
+			Value = ConvertedTime.FrameNumber;
+		}
+
+		return Value;
+	}
+
+	return FFrameNumber();
+}
+
+FFrameNumber ULevelSequencePlaybackController::GetCurrentSequencePlaybackEnd() const
+{
+	if (ActiveLevelSequence)
+	{
+		FFrameNumber Value = ActiveLevelSequence->GetMovieScene()->GetPlaybackRange().GetUpperBoundValue();
+		if (ActiveLevelSequence->GetMovieScene()->GetDisplayRate() != ActiveLevelSequence->GetMovieScene()->GetTickResolution())
+		{
+			const FFrameTime ConvertedTime = FFrameRate::TransformTime(FFrameTime(Value), ActiveLevelSequence->GetMovieScene()->GetTickResolution(), ActiveLevelSequence->GetMovieScene()->GetDisplayRate());
+			Value = ConvertedTime.FrameNumber;
+		}
+	
+		return Value;
+	}
+
+	return FFrameNumber();
+}
+
+FFrameNumber ULevelSequencePlaybackController::GetCurrentSequenceDuration() const
+{
+	if (ActiveLevelSequence)
+	{
+		FFrameNumber Value = UE::MovieScene::DiscreteSize(ActiveLevelSequence->GetMovieScene()->GetPlaybackRange());
+		if (ActiveLevelSequence->GetMovieScene()->GetDisplayRate() != ActiveLevelSequence->GetMovieScene()->GetTickResolution())
+		{
+			const FFrameTime ConvertedTime = FFrameRate::TransformTime(FFrameTime(Value), ActiveLevelSequence->GetMovieScene()->GetTickResolution(), ActiveLevelSequence->GetMovieScene()->GetDisplayRate());
+			Value = ConvertedTime.FrameNumber;
+		}
+		
+		return Value;
+	}
+
+	return FFrameNumber();
+}
+
+FFrameTime ULevelSequencePlaybackController::GetCurrentSequencePlaybackPosition() const
+{
+#if WITH_EDITOR
+	if (ActiveLevelSequence)
+	{
+		TSharedPtr<ISequencer> Sequencer = WeakSequencer.Pin();
+		if (Sequencer)
+		{
+			return Sequencer->GetLocalTime().ConvertTo(Sequencer->GetFocusedDisplayRate());
+		}
+	}
+#endif //WITH_EDITOR
+
+	return FFrameTime();
+}
+
+FTimecode ULevelSequencePlaybackController::GetCurrentSequencePlaybackTimecode() const
+{
+#if WITH_EDITOR
+	if (ActiveLevelSequence)
+	{
+		TSharedPtr<ISequencer> Sequencer = WeakSequencer.Pin();
+		if (Sequencer)
+		{
+			const FFrameTime DisplayTime = Sequencer->GetLocalTime().ConvertTo(Sequencer->GetFocusedDisplayRate());
+			return FTimecode::FromFrameNumber(DisplayTime.FrameNumber, Sequencer->GetFocusedDisplayRate());
+		}
+	}
+#endif //WITH_EDITOR
+
+	return FTimecode();
+}
+
+void ULevelSequencePlaybackController::JumpToPlaybackPosition(const FFrameNumber& InFrameNumber)
+{
+#if WITH_EDITOR
+	if (ActiveLevelSequence)
+	{
+		TSharedPtr<ISequencer> Sequencer = WeakSequencer.Pin();
+		if (Sequencer)
+		{
+			const FFrameTime NewTime = ConvertFrameTime(InFrameNumber, Sequencer->GetFocusedDisplayRate(), Sequencer->GetFocusedTickResolution());
+			Sequencer->SetLocalTime(NewTime, STM_None);
+		}
+	}
+#endif //WITH_EDITOR
+}
+
+bool ULevelSequencePlaybackController::IsSequencePlaybackActive() const 
+{
+#if WITH_EDITOR
+	if (ActiveLevelSequence)
+	{
+		TSharedPtr<ISequencer> Sequencer = WeakSequencer.Pin();
+		if (Sequencer)
+		{
+			return Sequencer->GetPlaybackStatus() == EMovieScenePlayerStatus::Playing && Sequencer->GetPlaybackSpeed() != 0.0f;
+		}
+	}
+#endif //WITH_EDITOR
+
+	return false;
+}
+
+void ULevelSequencePlaybackController::PauseLevelSequence()
+{
+#if WITH_EDITOR
+	if (ActiveLevelSequence)
+	{
+		TSharedPtr<ISequencer> Sequencer = WeakSequencer.Pin();
+		if (Sequencer)
+		{
+			Sequencer->Pause();
+		}
+	}
+#endif //WITH_EDITOR
+}
+
+void ULevelSequencePlaybackController::PlayLevelSequence()
+{
+#if WITH_EDITOR
+	if (ActiveLevelSequence)
+	{
+		TSharedPtr<ISequencer> Sequencer = WeakSequencer.Pin();
+		if (Sequencer)
+		{
+			Sequencer->SetPlaybackSpeed(1.0f);
+			Sequencer->OnPlay(false);
+		}
+	}
+#endif //WITH_EDITOR
+}
+
+void ULevelSequencePlaybackController::PlayLevelSequenceReverse()
+{
+#if WITH_EDITOR
+	if (ActiveLevelSequence)
+	{
+		TSharedPtr<ISequencer> Sequencer = WeakSequencer.Pin();
+		if (Sequencer)
+		{
+			Sequencer->SetPlaybackSpeed(-1.0f);
+			Sequencer->OnPlay(false);
+		}
+	}
+#endif //WITH_EDITOR
+}
+
+void ULevelSequencePlaybackController::StopLevelSequencePlay()
+{
+#if WITH_EDITOR
+	if (ActiveLevelSequence)
+	{
+		TSharedPtr<ISequencer> Sequencer = WeakSequencer.Pin();
+		if (Sequencer)
+		{
+			Sequencer->SetPlaybackStatus(EMovieScenePlayerStatus::Stopped);
+		}
+	}
+#endif //WITH_EDITOR
+}
+
+bool ULevelSequencePlaybackController::SetActiveLevelSequence(ULevelSequence* InNewLevelSequence)
+{
+#if WITH_EDITOR
+	if(InNewLevelSequence && !IsRunningGame())
+	{
+		const bool bDoFocusOnEditor = false;
+		GEditor->GetEditorSubsystem<UAssetEditorSubsystem>()->OpenEditorForAsset(InNewLevelSequence);
+		IAssetEditorInstance* AssetEditor = GEditor->GetEditorSubsystem<UAssetEditorSubsystem>()->FindEditorForAsset(InNewLevelSequence, bDoFocusOnEditor);
+		ILevelSequenceEditorToolkit* LevelSequenceEditor = static_cast<ILevelSequenceEditorToolkit*>(AssetEditor);
+
+		WeakSequencer = LevelSequenceEditor ? LevelSequenceEditor->GetSequencer() : nullptr;
+		if (WeakSequencer.IsValid())
+		{
+			TSharedPtr<ISequencer> Sequencer = WeakSequencer.Pin();
+			if (Sequencer)
+			{
+				const bool bSetCameraPerspective = InNewLevelSequence->GetMovieScene()->GetCameraCutTrack() != nullptr;
+
+				//If we're dealing with a sequence that already has a camera cut, don't enable recording and set camera perspective for review.
+				OnRecordEnabledStateChanged.ExecuteIfBound(!bSetCameraPerspective);
+				Sequencer->SetPerspectiveViewportCameraCutEnabled(bSetCameraPerspective);
+
+				ActiveLevelSequence = InNewLevelSequence;
+				return true;
+			}
+		}
+	}
+#endif //WITH_EDITOR
+
+	return false;
+}
+
+void ULevelSequencePlaybackController::ClearActiveLevelSequence()
+{
+	if (ActiveLevelSequence)
+	{
+		StopLevelSequencePlay();
+		ActiveLevelSequence = nullptr;
+#if WITH_EDITOR
+		WeakSequencer = nullptr;
+#endif
+	}
+}
+
+void ULevelSequencePlaybackController::PlayToEnd()
+{
+	if (ActiveLevelSequence)
+	{
+		PlayLevelSequence();
+	}
+}
diff --git a/Engine/Plugins/Experimental/VirtualCamera/Source/VirtualCamera/Private/LevelSequence/VirtualCameraClipsMetaData.cpp b/Engine/Plugins/Experimental/VirtualCamera/Source/VirtualCamera/Private/LevelSequence/VirtualCameraClipsMetaData.cpp
new file mode 100644
--- /dev/null
+++ b/Engine/Plugins/Experimental/VirtualCamera/Source/VirtualCamera/Private/LevelSequence/VirtualCameraClipsMetaData.cpp
@@ -0,0 +1,161 @@
+// Copyright Epic Games, Inc. All Rights Reserved.
+
+#include "LevelSequence/VirtualCameraClipsMetaData.h"
+
+const FName UVirtualCameraClipsMetaData::AssetRegistryTag_FocalLength = "ClipsMetaData_FocalLength";
+const FName UVirtualCameraClipsMetaData::AssetRegistryTag_bIsSelected = "ClipsMetaData_bIsSelected";
+const FName UVirtualCameraClipsMetaData::AssetRegistryTag_RecordedLevelName = "ClipsMetaData_RecordedLevelName";
+const FName UVirtualCameraClipsMetaData::AssetRegistryTag_FrameCountStart = "ClipsMetaData_FrameCountStart";
+const FName UVirtualCameraClipsMetaData::AssetRegistryTag_FrameCountEnd = "ClipsMetaData_FrameCountEnd"; 
+const FName UVirtualCameraClipsMetaData::AssetRegistryTag_LengthInFrames = "ClipsMetaData_LengthInFrames";
+const FName UVirtualCameraClipsMetaData::AssetRegistryTag_DisplayRate = "ClipsMetaData_DisplayRate";
+const FName UVirtualCameraClipsMetaData::AssetRegistryTag_bIsACineCameraRecording = "ClipsMetaData_bIsACineCameraRecording";
+
+
+UVirtualCameraClipsMetaData::UVirtualCameraClipsMetaData(const FObjectInitializer& ObjInit)
+	: Super(ObjInit)
+{
+	FocalLength = 0;
+	bIsSelected = false;
+	bIsACineCameraRecording = false; 
+}
+
+
+void UVirtualCameraClipsMetaData::ExtendAssetRegistryTags(TArray<FAssetRegistryTag>& OutTags) const
+{
+	FString IsSelected = bIsSelected ? "true" : "false";
+	FString IsRecordedFromACineCameraActor = bIsACineCameraRecording ? "true" : "false";
+
+	OutTags.Emplace(AssetRegistryTag_FocalLength, FString::FromInt(FocalLength), FAssetRegistryTag::ETagType::TT_Numerical, FAssetRegistryTag::TD_None);
+	OutTags.Emplace(AssetRegistryTag_bIsSelected, IsSelected, FAssetRegistryTag::ETagType::TT_Alphabetical, FAssetRegistryTag::TD_None);
+	OutTags.Emplace(AssetRegistryTag_RecordedLevelName, RecordedLevelName, FAssetRegistryTag::ETagType::TT_Alphabetical, FAssetRegistryTag::TD_None);
+	OutTags.Emplace(AssetRegistryTag_FrameCountStart, FString::FromInt(FrameCountStart), FAssetRegistryTag::ETagType::TT_Numerical, FAssetRegistryTag::TD_None);
+	OutTags.Emplace(AssetRegistryTag_FrameCountEnd, FString::FromInt(FrameCountEnd), FAssetRegistryTag::ETagType::TT_Numerical, FAssetRegistryTag::TD_None);
+	OutTags.Emplace(AssetRegistryTag_LengthInFrames, FString::FromInt(LengthInFrames), FAssetRegistryTag::ETagType::TT_Numerical, FAssetRegistryTag::TD_None);
+	OutTags.Emplace(AssetRegistryTag_DisplayRate, DisplayRate.ToPrettyText().ToString(), FAssetRegistryTag::ETagType::TT_Alphabetical, FAssetRegistryTag::TD_None);
+	OutTags.Emplace(AssetRegistryTag_bIsACineCameraRecording, IsRecordedFromACineCameraActor, FAssetRegistryTag::ETagType::TT_Alphabetical, FAssetRegistryTag::TD_None);
+}
+
+#if WITH_EDITOR
+void UVirtualCameraClipsMetaData::ExtendAssetRegistryTagMetaData(TMap<FName, FAssetRegistryTagMetadata>& OutMetadata) const
+{
+	OutMetadata.Add(AssetRegistryTag_FocalLength, FAssetRegistryTagMetadata()
+		.SetDisplayName(NSLOCTEXT("ClipsMetaData", "FocalLength_Label", "FocalLength"))
+		.SetTooltip(NSLOCTEXT("ClipsMetaData", "FocalLength_Tip", "The focal length that this level sequence was recorded with"))
+	);
+
+	OutMetadata.Add(AssetRegistryTag_bIsSelected, FAssetRegistryTagMetadata()
+		.SetDisplayName(NSLOCTEXT("ClipsMetaData", "IsSelected_Label", "bIsSelected"))
+		.SetTooltip(NSLOCTEXT("ClipsMetaData", "IsSelected_Tip", "If this level sequence is marked as a selected take"))
+	);
+
+	OutMetadata.Add(AssetRegistryTag_RecordedLevelName, FAssetRegistryTagMetadata()
+		.SetDisplayName(NSLOCTEXT("ClipsMetaData", "RecordedLavelName_Label", "RecordedLevelName"))
+		.SetTooltip(NSLOCTEXT("ClipsMetaData", "RecordedLavelName_Tip", "The name of the level this sequence was recorded"))
+	);
+
+	OutMetadata.Add(AssetRegistryTag_FrameCountStart, FAssetRegistryTagMetadata()
+		.SetDisplayName(NSLOCTEXT("ClipsMetaData", "FrameCountStart_Label", "FrameCountStart"))
+		.SetTooltip(NSLOCTEXT("ClipsMetaData", "FrameCountStart_Tip", "The first frame of the level sequence. "))
+	);
+
+	OutMetadata.Add(AssetRegistryTag_FrameCountEnd, FAssetRegistryTagMetadata()
+		.SetDisplayName(NSLOCTEXT("ClipsMetaData", "FrameCountEnd_Label", "FrameCountEnd"))
+		.SetTooltip(NSLOCTEXT("ClipsMetaData", "FrameCountEnd_Tip", "The last frame of the level sequence."))
+	);
+	OutMetadata.Add(AssetRegistryTag_LengthInFrames, FAssetRegistryTagMetadata()
+		.SetDisplayName(NSLOCTEXT("ClipsMetaData", "LengthInFrames_Label", "LengthInFrames"))
+		.SetTooltip(NSLOCTEXT("ClipsMetaData", "LengthInFrames_Tip", "The length in frames of the level sequence."))
+	);
+	OutMetadata.Add(AssetRegistryTag_DisplayRate, FAssetRegistryTagMetadata()
+		.SetDisplayName(NSLOCTEXT("ClipsMetaData", "DisplayRate_Label", "DisplayRate"))
+		.SetTooltip(NSLOCTEXT("ClipsMetaData", "DisplayRate_Tip", "The display rate of level sequence."))
+	);
+
+	OutMetadata.Add(AssetRegistryTag_bIsACineCameraRecording, FAssetRegistryTagMetadata()
+		.SetDisplayName(NSLOCTEXT("ClipsMetaData", "IsACineCameraRecording_Label", "bIsACineCameraRecording"))
+		.SetTooltip(NSLOCTEXT("ClipsMetaData", "IsACineCameraRecording_Tip", "If the level sequence was recorded with a CineCameraActor."))
+	);
+}
+#endif
+
+float UVirtualCameraClipsMetaData::GetFocalLength() const 
+{
+	return FocalLength; 
+}
+
+bool UVirtualCameraClipsMetaData::GetSelected() const
+{
+	return bIsSelected; 
+}
+
+FString UVirtualCameraClipsMetaData::GetRecordedLevelName() const
+{
+	return RecordedLevelName; 
+}
+
+int UVirtualCameraClipsMetaData::GetFrameCountStart() const
+{
+	return FrameCountStart; 
+}
+
+int UVirtualCameraClipsMetaData::GetFrameCountEnd() const
+{
+	return FrameCountEnd; 
+}
+
+int UVirtualCameraClipsMetaData::GetLengthInFrames()
+{
+	return LengthInFrames;
+}
+
+FFrameRate UVirtualCameraClipsMetaData::GetDisplayRate()
+{
+	return DisplayRate;
+}
+
+bool UVirtualCameraClipsMetaData::GetIsACineCameraRecording() const
+{
+	return bIsACineCameraRecording; 
+}
+
+void UVirtualCameraClipsMetaData::SetFocalLength(float InFocalLength) 
+{
+	FocalLength = InFocalLength;
+}
+
+void UVirtualCameraClipsMetaData::SetSelected(bool bInSelected)
+{
+	bIsSelected = bInSelected;
+}
+
+void UVirtualCameraClipsMetaData::SetRecordedLevelName(FString InLevelName)
+{
+	RecordedLevelName = InLevelName;
+}
+
+void UVirtualCameraClipsMetaData::SetFrameCountStart(int InFrame)
+{
+	FrameCountStart = InFrame; 
+}
+
+void UVirtualCameraClipsMetaData::SetFrameCountEnd(int InFrame)
+{
+	FrameCountEnd = InFrame; 
+}
+
+void UVirtualCameraClipsMetaData::SetLengthInFrames(int InLength)
+{
+	LengthInFrames = InLength; 
+}
+
+void UVirtualCameraClipsMetaData::SetDisplayRate(FFrameRate InDisplayRate)
+{
+	DisplayRate = InDisplayRate; 
+}
+
+void UVirtualCameraClipsMetaData::SetIsACineCameraRecording(bool bInIsACineCameraRecording)
+{
+	bIsACineCameraRecording = bInIsACineCameraRecording; 
+}
+
diff --git a/Engine/Plugins/Experimental/VirtualCamera/Source/VirtualCamera/Private/LevelSequencePlaybackController.cpp b/Engine/Plugins/Experimental/VirtualCamera/Source/VirtualCamera/Private/LevelSequencePlaybackController.cpp
deleted file mode 100644
diff --git a/Engine/Plugins/Experimental/VirtualCamera/Source/VirtualCamera/Private/VCamBlueprintFunctionLibrary.cpp b/Engine/Plugins/Experimental/VirtualCamera/Source/VirtualCamera/Private/VCamBlueprintFunctionLibrary.cpp
--- a/Engine/Plugins/Experimental/VirtualCamera/Source/VirtualCamera/Private/VCamBlueprintFunctionLibrary.cpp
+++ b/Engine/Plugins/Experimental/VirtualCamera/Source/VirtualCamera/Private/VCamBlueprintFunctionLibrary.cpp
@@ -2,15 +2,16 @@
 
 #include "VCamBlueprintFunctionLibrary.h"
 
+#include "LevelSequence/VirtualCameraClipsMetaData.h"
+#include "VirtualCameraUserSettings.h"
+
+#include "AssetRegistry/AssetData.h"
 #include "CineCameraActor.h"
 #include "CineCameraComponent.h"
-#include "Engine/GameInstance.h"
-#include "LevelSequence.h"
-#include "VirtualCameraClipsMetaData.h"
 #include "Components/SceneCaptureComponent2D.h"
-#include "AssetRegistry/AssetData.h"
-#include "VirtualCameraUserSettings.h"
+#include "Engine/GameInstance.h"
 #include "GameFramework/PlayerController.h"
+#include "LevelSequence.h"
 #include "MovieScene.h"
 
 #if WITH_EDITOR
@@ -27,7 +28,6 @@
 #include "Recorder/TakeRecorderPanel.h"
 #include "SLevelViewport.h"
 #include "SceneView.h"
-#include "Subsystems/UnrealEditorSubsystem.h"
 #include "TakePreset.h"
 #include "VPUtilitiesEditorBlueprintLibrary.h"
 #endif
diff --git a/Engine/Plugins/Experimental/VirtualCamera/Source/VirtualCamera/Private/VirtualCameraClipsMetaData.cpp b/Engine/Plugins/Experimental/VirtualCamera/Source/VirtualCamera/Private/VirtualCameraClipsMetaData.cpp
deleted file mode 100644
diff --git a/Engine/Plugins/Experimental/VirtualCamera/Source/VirtualCamera/Private/VirtualCameraSubsystem.cpp b/Engine/Plugins/Experimental/VirtualCamera/Source/VirtualCamera/Private/VirtualCameraSubsystem.cpp
--- a/Engine/Plugins/Experimental/VirtualCamera/Source/VirtualCamera/Private/VirtualCameraSubsystem.cpp
+++ b/Engine/Plugins/Experimental/VirtualCamera/Source/VirtualCamera/Private/VirtualCameraSubsystem.cpp
@@ -1,7 +1,9 @@
 // Copyright Epic Games, Inc. All Rights Reserved.
 
 #include "VirtualCameraSubsystem.h"
-#include "LevelSequencePlaybackController.h"
+
+#include "LevelSequence/LevelSequencePlaybackController.h"
+
 #include "Engine/Selection.h"
 #include "GameFramework/Actor.h"
 
diff --git a/Engine/Plugins/Experimental/VirtualCamera/Source/VirtualCamera/Public/LevelSequence/LevelSequencePlaybackController.h b/Engine/Plugins/Experimental/VirtualCamera/Source/VirtualCamera/Public/LevelSequence/LevelSequencePlaybackController.h
new file mode 100644
--- /dev/null
+++ b/Engine/Plugins/Experimental/VirtualCamera/Source/VirtualCamera/Public/LevelSequence/LevelSequencePlaybackController.h
@@ -0,0 +1,188 @@
+// Copyright Epic Games, Inc. All Rights Reserved.
+
+#pragma once
+
+#include "UObject/Object.h"
+
+#include "CoreMinimal.h"
+#include "Misc/FrameNumber.h"
+#include "Misc/FrameRate.h"
+#include "Misc/Timecode.h"
+
+#if WITH_EDITOR
+#include "ISequencer.h"
+#endif
+
+#include "LevelSequencePlaybackController.generated.h"
+
+
+class ULevelSequence;
+
+
+DECLARE_DELEGATE_OneParam(FRecordEnabledStateChanged, bool)
+
+USTRUCT(BlueprintType)
+struct FLevelSequenceData
+{
+	GENERATED_BODY()
+
+public:
+	UPROPERTY(EditAnywhere, BlueprintReadWrite, Category = "Level Sequence Info")
+	FString AssetPath;
+
+	UPROPERTY(EditAnywhere, BlueprintReadWrite, Category = "Level Sequence Info")
+	FString DisplayName;
+
+	UPROPERTY(EditAnywhere, BlueprintReadWrite, Category = "Level Sequence Info", meta = (IgnoreForMemberInitializationTest))
+	FDateTime LastEdited;
+
+	FLevelSequenceData(const FString InAssetPath = "", const FString InDisplayName = "", const FDateTime InLastEdited = FDateTime::Now())
+	{
+		AssetPath = InAssetPath;
+		DisplayName = InDisplayName;
+		LastEdited = InLastEdited;
+	}
+};
+
+UCLASS(BlueprintType)
+class VIRTUALCAMERA_API ULevelSequencePlaybackController : public UObject
+{
+	GENERATED_UCLASS_BODY()
+
+public:
+	/** Notify whether recording is enabled or disabled for the current sequence*/
+	FRecordEnabledStateChanged OnRecordEnabledStateChanged;
+
+	/**
+	 * Returns the names of each level sequence actor that is present in the level.
+	 * @param &OutLevelSequenceNames - UponReturn, array will contain all level sequence names.
+	 */
+	UFUNCTION(BlueprintCallable, Category = "VirtualCamera | Sequencer | Playback")
+	void GetLevelSequences(TArray<FLevelSequenceData>& OutLevelSequenceNames);
+
+	/**
+	 * @return the name of the currently selected sequence; returns empty string if no selected sequence
+	 */
+	UFUNCTION(BlueprintPure, Category = "VirtualCamera | Sequencer | Playback")
+	FString GetActiveLevelSequenceName() const;
+
+	/**
+	 * @return the currently selected LevelSequence
+	 */
+	UFUNCTION(BlueprintPure, Category = "VirtualCamera | Sequencer | Playback")
+	ULevelSequence* GetActiveLevelSequence() { return ActiveLevelSequence; }
+
+	/**
+	 * @return the FrameRate of the currently loaded sequence FrameRate
+	 */
+	UFUNCTION(BlueprintPure, Category = "VirtualCamera | Sequencer | Playback")
+	FFrameRate GetCurrentSequenceFrameRate() const;
+
+	/**
+	 * @return true if the active Sequencer is locked to camera cut
+	 */
+	UFUNCTION(BlueprintPure, Category = "VirtualCamera | Sequencer | Playback")
+	bool IsSequencerLockedToCameraCut() const;
+
+	/**
+	 * Sets the current Sequencer perspective to be locked to camera cut
+	 */
+	UFUNCTION(BlueprintCallable, Category = "VirtualCamera | Sequencer | Playback")
+	void SetSequencerLockedToCameraCut(bool bLockView);
+
+	/**
+	 * @return The FrameNumber of the sequence's start.
+	 */
+	UFUNCTION(BlueprintPure, Category = "VirtualCamera | Sequencer | Playback")
+	FFrameNumber GetCurrentSequencePlaybackStart() const;
+
+	/**
+	 * @return The FrameNumber of the sequence's end.
+	 */
+	UFUNCTION(BlueprintPure, Category = "VirtualCamera | Sequencer | Playback")
+	FFrameNumber GetCurrentSequencePlaybackEnd() const;
+
+	/**
+	 * @return the duration of the sequence in FrameNumber
+	 */
+	UFUNCTION(BlueprintPure, Category = "VirtualCamera | Sequencer | Playback")
+	FFrameNumber GetCurrentSequenceDuration() const;
+
+	/**
+	 * @return the current FrameTime of the sequence playback.
+	 */
+	UFUNCTION(BlueprintPure, Category = "VirtualCamera | Sequencer | Playback")
+	FFrameTime GetCurrentSequencePlaybackPosition() const;
+
+	/**
+	 * @return The current Timecode of the sequence playback.
+	 */
+	UFUNCTION(BlueprintPure, Category = "VirtualCamera | Sequencer | Playback")
+	FTimecode GetCurrentSequencePlaybackTimecode() const;
+
+	/**
+	 * Moves the current sequence to a desired playback position
+	 */
+	UFUNCTION(BlueprintCallable, Category = "VirtualCamera | Sequencer | Playback")
+	void JumpToPlaybackPosition(const FFrameNumber& InFrameNumber);
+
+	/**
+	 * @return true if a valid LevelSequence is being played.
+	 */
+	UFUNCTION(BlueprintPure, Category = "VirtualCamera | Sequencer | Playback")
+	bool IsSequencePlaybackActive() const;
+
+	/**
+	 * Pause the currently active sequence 
+	 */
+	UFUNCTION(BlueprintCallable, Category = "VirtualCamera | Sequencer | Playback")
+	void PauseLevelSequence();
+
+	/**
+	 * Starts playing the currently active sequence
+	 */
+	UFUNCTION(BlueprintCallable, Category = "VirtualCamera | Sequencer | Playback")
+	void PlayLevelSequence();
+
+	/**
+	 * Starts playing the currently active sequence in reverse
+	 */
+	UFUNCTION(BlueprintCallable, Category = "VirtualCamera | Sequencer | Playback")
+	void PlayLevelSequenceReverse();
+
+	/**
+	 * Stops playing the currently active sequence
+	 */
+	UFUNCTION(BlueprintCallable, Category = "VirtualCamera | Sequencer | Playback")
+	void StopLevelSequencePlay();
+
+	/**
+	 * Changes the active level sequence to a new level sequence.
+	 * @param InNewLevelSequence - The LevelSequence to select
+	 * @return true if a valid LevelSequence was passed and sequencer was successfully found.
+	 */
+	UFUNCTION(BlueprintCallable, Category = "VirtualCamera | Sequencer | Playback")
+	bool SetActiveLevelSequence(ULevelSequence* InNewLevelSequence);
+
+	/**
+	 * Clears the current level sequence player, needed when recording clean takes of something.
+	 */
+	UFUNCTION(BlueprintCallable, Category = "VirtualCamera | Sequencer | Playback")
+	void ClearActiveLevelSequence();
+
+protected:
+
+	/** The sequence to play back */
+	UPROPERTY(Transient)
+	TObjectPtr<ULevelSequence> ActiveLevelSequence;
+
+#if WITH_EDITORONLY_DATA
+	/** Weak reference to Sequencer associated with the active LevelSequence */
+	TWeakPtr<ISequencer> WeakSequencer;
+#endif //WITH_EDITORONLY_DATA
+
+	/**
+	 * Play to the end of the current sequence and stop
+	 */
+	void PlayToEnd();
+};
diff --git a/Engine/Plugins/Experimental/VirtualCamera/Source/VirtualCamera/Public/LevelSequence/VirtualCameraClipsMetaData.h b/Engine/Plugins/Experimental/VirtualCamera/Source/VirtualCamera/Public/LevelSequence/VirtualCameraClipsMetaData.h
new file mode 100644
--- /dev/null
+++ b/Engine/Plugins/Experimental/VirtualCamera/Source/VirtualCamera/Public/LevelSequence/VirtualCameraClipsMetaData.h
@@ -0,0 +1,194 @@
+// Copyright Epic Games, Inc. All Rights Reserved.
+
+#pragma once
+
+
+#include "UObject/Object.h"
+#include "ILevelSequenceMetaData.h"
+#include "Misc/FrameRate.h"
+#include "VirtualCameraClipsMetaData.generated.h"
+
+
+/**
+ * Clips meta-data that is stored on ULevelSequence assets that are recorded through the virtual camera. 
+ * Meta-data is retrieved through ULevelSequence::FindMetaData<UVirtualCameraClipsMetaData>()
+ */
+
+UCLASS(BlueprintType)
+class UVirtualCameraClipsMetaData : public UObject, public ILevelSequenceMetaData
+{
+public: 
+	GENERATED_BODY()
+	UVirtualCameraClipsMetaData(const FObjectInitializer& ObjInit);
+
+public: 
+
+	/** The asset registry tag that contains the focal length for this meta-data */
+	static const FName AssetRegistryTag_FocalLength;
+
+	/** The asset registry tag that contains if the selected state for this meta-data */
+	static const FName AssetRegistryTag_bIsSelected;
+
+	/** The asset registry tag that contains the recorded level name for this meta-data */
+	static const FName AssetRegistryTag_RecordedLevelName;
+
+	/** The asset registry tag that contains the FrameCountStart in for this meta-data */
+	static const FName AssetRegistryTag_FrameCountStart;
+
+	/** The asset registry tag that contains the FrameCountEnd out for this meta-data */
+	static const FName AssetRegistryTag_FrameCountEnd;
+
+	/** The asset registry tag that contains the LengthInFrames out for this meta-data */
+	static const FName AssetRegistryTag_LengthInFrames; 
+
+	/** The asset registry tag that contains the FrameCountEnd out for this meta-data */
+	static const FName AssetRegistryTag_DisplayRate;
+
+	/** The asset registry tag that contains whether the clip was recorded with a CineCamera for this meta-data */
+	static const FName AssetRegistryTag_bIsACineCameraRecording;
+
+
+public:
+
+	/**
+	* Extend the default ULevelSequence asset registry tags
+	*/
+	virtual void ExtendAssetRegistryTags(TArray<FAssetRegistryTag>& OutTags) const override;
+#if WITH_EDITOR
+	virtual void ExtendAssetRegistryTagMetaData(TMap<FName, FAssetRegistryTagMetadata>& OutMetadata) const override;
+#endif
+
+public:
+	/**
+	 * @return The focal length for this clip
+	 */
+	UFUNCTION(BlueprintCallable, Category = "Clips")
+	float GetFocalLength() const;
+
+	/**
+	* @return Whether or not the clip is selected.  
+	*/
+	UFUNCTION(BlueprintCallable, Category = "Clips")
+	bool GetSelected() const;
+
+	/**
+	* @return The name of the clip's recorded level. 
+	*/
+	UFUNCTION(BlueprintCallable, Category = "Clips")
+	FString GetRecordedLevelName() const;
+
+	/**
+	* @return The initial frame of the clip
+	*/
+	UFUNCTION(BlueprintCallable, Category = "Clips")
+	int GetFrameCountStart() const;
+
+	/**
+	* @return The final frame of the clip
+	*/
+	UFUNCTION(BlueprintCallable, Category = "Clips")
+	int GetFrameCountEnd() const;
+
+	/**
+	* @return The length in frames of the clip. 
+	*/
+	UFUNCTION(BlueprintCallable, Category = "Clips")
+	int GetLengthInFrames();
+
+	/**
+	* @return The display rate of the clip. 
+	*/
+	UFUNCTION(BlueprintCallable, Category = "Clips")
+	FFrameRate GetDisplayRate();
+
+	/**
+	* @return Whether the clip was recorded by a CineCameraActor
+	*/
+	UFUNCTION(BlueprintCallable, Category = "Clips")
+	bool GetIsACineCameraRecording() const;
+
+public:
+	/**
+	* Set the focal length associated with this clip. 
+	* @note: Used for tracking. Does not update the StreamedCameraComponent. 
+	*/
+	UFUNCTION(BlueprintCallable, Category = "Clips")
+	void SetFocalLength(float InFocalLength);
+	
+	/**
+	* Set if this clip is 'selected'
+	*/
+	UFUNCTION(BlueprintCallable, Category = "Clips")
+	void SetSelected(bool bInSelected);
+
+	/**
+	* Set the name of the level that the clip was recorded in. 
+	*/
+	UFUNCTION(BlueprintCallable, Category = "Clips")
+	void SetRecordedLevelName(FString InLevelName);
+
+	/**
+	* Set the initial frame of the clip used for calculating duration.
+	*/
+	UFUNCTION(BlueprintCallable, Category = "Clips")
+	void SetFrameCountStart(int InFrame); 
+
+	/**
+	* Set the final frame of the clip used for calculating duration.
+	*/
+	UFUNCTION(BlueprintCallable, Category = "Clips")
+	void SetFrameCountEnd(int InFrame);
+
+	/**
+	 * Set the length in frames of the clip used for AssetData calculations. 
+	 */
+	UFUNCTION(BlueprintCallable, Category = "Clips")
+	void SetLengthInFrames(int InLength);
+
+	/**
+	 * Set the DisplayRate of the clip used for AssetData calculations.
+	 */
+	UFUNCTION(BlueprintCallable, Category = "Clips")
+	void SetDisplayRate(FFrameRate InDisplayRate);
+
+	/**
+	 * Set if the clip was recorded by a CineCameraActor
+	 */
+	UFUNCTION(BlueprintCallable, Category = "Clips")
+	void SetIsACineCameraRecording(bool bInIsACineCameraRecording);
+
+
+private:
+	/** The focal length of the streamed camera used to record the take */
+	UPROPERTY()
+	float FocalLength;
+
+	/** Whether or not the take was marked as 'selected' */
+	UPROPERTY()
+	bool bIsSelected; 
+
+	/** The name of the level that the clip was recorded in */
+	UPROPERTY()
+	FString RecordedLevelName; 
+
+	/** The initial frame of the clip used for calculating duration. */
+	UPROPERTY()
+	int FrameCountStart;
+
+	/** The last frame of the clip used for calculating duration. */
+	UPROPERTY()
+	int FrameCountEnd; 
+
+	/** The level sequence length in frames calculated from VirtualCameraSubsystem used for AssetData calculations */
+	UPROPERTY()
+	int LengthInFrames; 
+
+	/** The display rate of the level sequence used for AssetData calculations. */
+	UPROPERTY()
+	FFrameRate DisplayRate;
+
+	/** If the LevelSequence was recorded with a CineCameraActor, rather than a VirtualCameraActor */
+	UPROPERTY()
+	bool bIsACineCameraRecording;
+
+};
\ No newline at end of file
diff --git a/Engine/Plugins/Experimental/VirtualCamera/Source/VirtualCamera/Public/LevelSequencePlaybackController.h b/Engine/Plugins/Experimental/VirtualCamera/Source/VirtualCamera/Public/LevelSequencePlaybackController.h
deleted file mode 100644
diff --git a/Engine/Plugins/Experimental/VirtualCamera/Source/VirtualCamera/Public/VirtualCameraClipsMetaData.h b/Engine/Plugins/Experimental/VirtualCamera/Source/VirtualCamera/Public/VirtualCameraClipsMetaData.h
deleted file mode 100644
diff --git a/Engine/Plugins/Experimental/VirtualCamera/Source/VirtualCamera/Public/VirtualCameraPlayerControllerBase.h b/Engine/Plugins/Experimental/VirtualCamera/Source/VirtualCamera/Public/VirtualCameraPlayerControllerBase.h
--- a/Engine/Plugins/Experimental/VirtualCamera/Source/VirtualCamera/Public/VirtualCameraPlayerControllerBase.h
+++ b/Engine/Plugins/Experimental/VirtualCamera/Source/VirtualCamera/Public/VirtualCameraPlayerControllerBase.h
@@ -4,15 +4,15 @@
 
 #include "GameFramework/PlayerController.h"
 
-#include "CineCameraActor.h"
-#include "InputCore.h"
 #include "IVirtualCameraController.h"
-#include "LevelSequencePlaybackController.h"
+#include "LevelSequence/LevelSequencePlaybackController.h"
+#include "VirtualCameraPawnBase.h"
+
+#include "CineCameraActor.h"
 #include "Math/UnitConversion.h"
 #include "Misc/FrameNumber.h"
 #include "Misc/FrameRate.h"
 #include "Misc/Timecode.h"
-#include "VirtualCameraPawnBase.h"
 #include "VirtualCameraPlayerControllerBase.generated.h"
 
 class AVPRootActor;
