Add support for PixelStreamingVCam to resize UE viewport to match remote device resolution.

diff --git a/Engine/Plugins/Experimental/VirtualCamera/Source/PixelStreamingVCam/Private/PixelStreamingMediaOutput.cpp b/Engine/Plugins/Experimental/VirtualCamera/Source/PixelStreamingVCam/Private/PixelStreamingMediaOutput.cpp
--- a/Engine/Plugins/Experimental/VirtualCamera/Source/PixelStreamingVCam/Private/PixelStreamingMediaOutput.cpp
+++ b/Engine/Plugins/Experimental/VirtualCamera/Source/PixelStreamingVCam/Private/PixelStreamingMediaOutput.cpp
@@ -9,6 +9,7 @@
 #include "PixelStreamingEditorModule.h"
 #include "PixelStreamingEditorUtils.h"
 #include "PixelStreamingVideoInputRHI.h"
+#include "PixelStreamingUtils.h"
 
 void UPixelStreamingMediaOutput::BeginDestroy()
 {
@@ -23,6 +24,7 @@ UMediaCapture* UPixelStreamingMediaOutput::CreateMediaCaptureImpl()
 	{
 		IPixelStreamingModule& Module = FModuleManager::LoadModuleChecked<IPixelStreamingModule>("PixelStreaming");
 		Streamer = Module.GetStreamer(Module.GetDefaultStreamerID());
+		RegisterRemoteResolutionCommandHandler();
 	}
 
 	Capture = nullptr;
@@ -42,6 +44,31 @@ UMediaCapture* UPixelStreamingMediaOutput::CreateMediaCaptureImpl()
 	return Capture;
 }
 
+void UPixelStreamingMediaOutput::RegisterRemoteResolutionCommandHandler()
+{
+	// Override resolution command as we this to set the output provider override resolution
+	TSharedPtr<IPixelStreamingInputHandler> InputHandler = Streamer->GetInputHandler().Pin();
+	if(InputHandler)
+	{
+		InputHandler->SetCommandHandler(TEXT("Resolution.Width"), [this](FString Descriptor, FString WidthString){
+			bool bSuccess = false;
+			FString HeightString;
+			UE::PixelStreaming::ExtractJsonFromDescriptor(Descriptor, TEXT("Resolution.Height"), HeightString, bSuccess);
+			if (bSuccess)
+			{
+				int Width = FCString::Atoi(*WidthString);
+				int Height = FCString::Atoi(*HeightString);
+				if (Width < 1 || Height < 1)
+				{
+					return;
+				}
+
+				RemoteResolutionChangedEvent.Broadcast(FIntPoint(Width, Height));
+			}
+		});
+	}
+}
+
 void UPixelStreamingMediaOutput::StartStreaming()
 {
 	if (Streamer)
diff --git a/Engine/Plugins/Experimental/VirtualCamera/Source/PixelStreamingVCam/Private/PixelStreamingMediaOutput.h b/Engine/Plugins/Experimental/VirtualCamera/Source/PixelStreamingVCam/Private/PixelStreamingMediaOutput.h
--- a/Engine/Plugins/Experimental/VirtualCamera/Source/PixelStreamingVCam/Private/PixelStreamingMediaOutput.h
+++ b/Engine/Plugins/Experimental/VirtualCamera/Source/PixelStreamingVCam/Private/PixelStreamingMediaOutput.h
@@ -6,6 +6,7 @@
 #include "IPixelStreamingStreamer.h"
 #include "PixelStreamingMediaCapture.h"
 #include "PixelStreamingVideoInput.h"
+#include "Delegates/DelegateCombinations.h"
 #include "PixelStreamingMediaOutput.generated.h"
 
 
@@ -28,20 +29,28 @@ protected:
 
 public:
 	TSharedPtr<IPixelStreamingStreamer> GetStreamer() const { return Streamer; }
-
 	void SetSignallingServerURL(FString InURL);
 	void SetSignallingStreamID(FString InStreamID);
 	void StartStreaming();
 	void StopStreaming();
 
+	DECLARE_EVENT_OneParam(UPixelStreamingMediaOutput, FRemoteResolutionChangedEvent, const FIntPoint&)
+	FRemoteResolutionChangedEvent& OnRemoteResolutionChanged() { return RemoteResolutionChangedEvent; }
+
 private:
 	UPixelStreamingMediaCapture* Capture = nullptr;
 	TSharedPtr<IPixelStreamingStreamer> Streamer;
 	TSharedPtr<FPixelStreamingVideoInput> VideoInput;
 
+	/** Broadcasts whenever the layer changes */
+	FRemoteResolutionChangedEvent RemoteResolutionChangedEvent;
+
 	FString SignallingServerURL;
 	FString StreamID = TEXT("VCam");
 
 	void OnCaptureStateChanged();
 	void OnCaptureViewportInitialized();
+
+	/* When the remote device sends a resolution change request back to the VCam we can respond here. */
+	void RegisterRemoteResolutionCommandHandler();
 };
diff --git a/Engine/Plugins/Experimental/VirtualCamera/Source/PixelStreamingVCam/Private/VCamPixelStreamingSession.cpp b/Engine/Plugins/Experimental/VirtualCamera/Source/PixelStreamingVCam/Private/VCamPixelStreamingSession.cpp
--- a/Engine/Plugins/Experimental/VirtualCamera/Source/PixelStreamingVCam/Private/VCamPixelStreamingSession.cpp
+++ b/Engine/Plugins/Experimental/VirtualCamera/Source/PixelStreamingVCam/Private/VCamPixelStreamingSession.cpp
@@ -75,6 +75,7 @@ void UVCamPixelStreamingSession::Activate()
 	if (MediaOutput == nullptr)
 	{
 		MediaOutput = NewObject<UPixelStreamingMediaOutput>(GetTransientPackage(), UPixelStreamingMediaOutput::StaticClass());
+		MediaOutput->OnRemoteResolutionChanged().AddUObject(this, &UVCamPixelStreamingSession::OnRemoteResolutionChanged);
 	}
 
 	UEditorPerformanceSettings* Settings = GetMutableDefault<UEditorPerformanceSettings>();
@@ -154,6 +155,25 @@ void UVCamPixelStreamingSession::OnCaptureStateChanged()
 	}
 }
 
+void UVCamPixelStreamingSession::OnRemoteResolutionChanged(const FIntPoint& RemoteResolution)
+{
+	// Early out if match remote resolution is not enabled.
+	if(!bMatchRemoteResolution)
+	{
+		return;
+	}
+
+	// Ensure override resolution is being used
+	if(!bUseOverrideResolution)
+	{
+		bUseOverrideResolution = true;
+	}
+
+	// Set the override resolution on the output provider base, this will trigger a resize
+	OverrideResolution = RemoteResolution;
+	ApplyOverrideResolutionForViewport(GetTargetViewport());
+}
+
 void UVCamPixelStreamingSession::SetupCustomInputHandling()
 {
 	// If we have a UMG, then use it
diff --git a/Engine/Plugins/Experimental/VirtualCamera/Source/PixelStreamingVCam/Public/VCamPixelStreamingSession.h b/Engine/Plugins/Experimental/VirtualCamera/Source/PixelStreamingVCam/Public/VCamPixelStreamingSession.h
--- a/Engine/Plugins/Experimental/VirtualCamera/Source/PixelStreamingVCam/Public/VCamPixelStreamingSession.h
+++ b/Engine/Plugins/Experimental/VirtualCamera/Source/PixelStreamingVCam/Public/VCamPixelStreamingSession.h
@@ -29,6 +29,10 @@ public:
 	UPROPERTY(EditAnywhere, BlueprintReadWrite, Category = "Output")
 	int32 FromComposureOutputProviderIndex = INDEX_NONE;
 
+	// If true the streamed UE viewport will match the resolution of the remote device.
+	UPROPERTY(EditAnywhere, BlueprintReadWrite, Category = "Output")
+	bool bMatchRemoteResolution = true;
+
 	// Check this if you wish to control the corresponding CineCamera with transform data received from the LiveLink app
 	UPROPERTY(EditAnywhere, BlueprintReadWrite, Category = "Output")
 	bool EnableARKitTracking = true;
@@ -60,6 +64,7 @@ private:
 	void SetupCustomInputHandling();
 	void OnCaptureStateChanged();
 	void OnARKitTransformReceived(FPixelStreamingPlayerId PlayerId, uint8 Type, TArray<uint8> Data);
+	void OnRemoteResolutionChanged(const FIntPoint& RemoteResolution);
 private:
 	FHitResult 	LastViewportTouchResult;
 	bool 		bUsingDummyUMG = false;
