Provide a warning and a correction path for inputs used in modifiers

diff --git a/Engine/Plugins/Experimental/VirtualCamera/Source/VCamCoreEditor/Private/Compilation/CompilationExtensionManager.cpp b/Engine/Plugins/Experimental/VirtualCamera/Source/VCamCoreEditor/Private/Compilation/CompilationExtensionManager.cpp
new file mode 100644
--- /dev/null
+++ b/Engine/Plugins/Experimental/VirtualCamera/Source/VCamCoreEditor/Private/Compilation/CompilationExtensionManager.cpp
@@ -0,0 +1,56 @@
+// Copyright Epic Games, Inc. All Rights Reserved.
+
+#include "CompilationExtensionManager.h"
+
+#include "ModifierCompilationBlueprintExtension.h"
+#include "Kismet2/KismetEditorUtilities.h"
+#include "Modifier/VCamModifier.h"
+
+namespace UE::VCamCoreEditor::Private
+{
+	FCompilationExtensionManager::~FCompilationExtensionManager()
+	{
+		FCoreUObjectDelegates::OnAssetLoaded.RemoveAll(this);
+		
+	}
+
+	void FCompilationExtensionManager::Init()
+	{
+		// This function exists because AddSP cannot be used in the constructor.
+		FCoreUObjectDelegates::OnAssetLoaded.AddSP(this, &FCompilationExtensionManager::OnAssetLoaded);
+	}
+
+	void FCompilationExtensionManager::OnAssetLoaded(UObject* Object) const
+	{
+		UBlueprint* Blueprint = Cast<UBlueprint>(Object);
+		if (!Blueprint)
+		{
+			return;
+		}
+
+		const bool bIsModifierBlueprint = Blueprint->ParentClass && Blueprint->ParentClass->IsChildOf(UVCamModifier::StaticClass());
+		if (!bIsModifierBlueprint)
+		{
+			return;
+		}
+
+		TArrayView<const TObjectPtr<UBlueprintExtension>> Extensions = Blueprint->GetExtensions();
+		const int32 Index = Extensions.IndexOfByPredicate([](UBlueprintExtension* Extension){ return Extension->IsA<UModifierCompilationBlueprintExtension>(); });
+		UModifierCompilationBlueprintExtension* Extension;
+		if (Index == INDEX_NONE)
+		{
+			Extension = NewObject<UModifierCompilationBlueprintExtension>(Blueprint);
+			Blueprint->AddExtension(Extension);
+		}
+		else
+		{
+			Extension = Cast<UModifierCompilationBlueprintExtension>(Extensions[Index].Get());
+		}
+		
+		// Blueprints are compiled during the loading (editor) screen. The FLinker loads the InputMappingContext reference but will only serialize it AFTER the initial compilation hence we request a recompilation here if necessary.
+		if (Extension->RequiresRecompileToDetectIssues(*Blueprint))
+		{
+			FKismetEditorUtilities::CompileBlueprint(Blueprint);
+		}
+	}
+}
diff --git a/Engine/Plugins/Experimental/VirtualCamera/Source/VCamCoreEditor/Private/Compilation/CompilationExtensionManager.h b/Engine/Plugins/Experimental/VirtualCamera/Source/VCamCoreEditor/Private/Compilation/CompilationExtensionManager.h
new file mode 100644
--- /dev/null
+++ b/Engine/Plugins/Experimental/VirtualCamera/Source/VCamCoreEditor/Private/Compilation/CompilationExtensionManager.h
@@ -0,0 +1,26 @@
+// Copyright Epic Games, Inc. All Rights Reserved.
+
+#pragma once
+
+#include "Templates/SharedPointer.h"
+
+class UModifierCompilationBlueprintExtension;
+
+namespace UE::VCamCoreEditor::Private
+{
+	/** Manages all of the plugin's UBlueprintExtensions making sure they are added when a Blueprint asset is loaded. */
+	class FCompilationExtensionManager : public TSharedFromThis<FCompilationExtensionManager>
+	{
+	public:
+
+		~FCompilationExtensionManager();
+
+		void Init();
+
+	private:
+
+		void OnAssetLoaded(UObject* Object) const;
+	};
+}
+
+
diff --git a/Engine/Plugins/Experimental/VirtualCamera/Source/VCamCoreEditor/Private/Compilation/ModifierCompilationBlueprintExtension.cpp b/Engine/Plugins/Experimental/VirtualCamera/Source/VCamCoreEditor/Private/Compilation/ModifierCompilationBlueprintExtension.cpp
new file mode 100644
--- /dev/null
+++ b/Engine/Plugins/Experimental/VirtualCamera/Source/VCamCoreEditor/Private/Compilation/ModifierCompilationBlueprintExtension.cpp
@@ -0,0 +1,141 @@
+// Copyright Epic Games, Inc. All Rights Reserved.
+
+#include "ModifierCompilationBlueprintExtension.h"
+
+#include "Modifier/VCamModifier.h"
+
+#include "EdGraph/EdGraphNode.h"
+#include "Engine/BlueprintGeneratedClass.h"
+#include "InputAction.h"
+#include "InputMappingContext.h"
+#include "K2Node_EnhancedInputAction.h"
+#include "Kismet2/CompilerResultsLog.h"
+#include "KismetCompiler.h"
+#include "Logging/TokenizedMessage.h"
+#include "Misc/UObjectToken.h"
+
+#define LOCTEXT_NAMESPACE "UModifierCompilationBlueprintExtension"
+
+namespace UE::VCamCoreEditor::Private
+{
+	enum class EBreakBehaviour
+	{
+		Break,
+		Continue
+	};
+	
+	/**
+	 * @param Blueprint Blueprint of which to iterate all nodes
+	 * @param Callback has this signature EBreakBehaviour(UEdGraphNode& Node)
+	 */
+	template<typename TCallback>
+	static void ForEachNode(UBlueprint& Blueprint, TCallback Callback)
+	{
+		TArray<UEdGraph*> Graphs;
+		Blueprint.GetAllGraphs(Graphs);
+		TSet<UEdGraphNode*> AlreadyVisited;
+		for (UEdGraph* Graph : Graphs)
+		{
+			for (UEdGraphNode* GraphNode : Graph->Nodes)
+			{
+				if (!GraphNode || AlreadyVisited.Contains(GraphNode))
+				{
+					continue;
+				}
+
+				AlreadyVisited.Add(GraphNode);
+				if (Callback(*GraphNode) == EBreakBehaviour::Break)
+				{
+					return;
+				}
+			}
+		}
+	}
+}
+
+void UModifierCompilationBlueprintExtension::HandleGenerateFunctionGraphs(FKismetCompilerContext* CompilerContext)
+{
+	UVCamModifier* OldModifierCDO = CompilerContext && CompilerContext->OldClass
+		? Cast<UVCamModifier>(CompilerContext->OldClass->GetDefaultObject())
+		: nullptr;
+	if (!OldModifierCDO)
+	{
+		return;
+	}
+
+	// Blueprints are compiled during the loading (editor) screen. The FLinker loads the InputMappingContext reference but will only serialize it AFTER the initial compilation making our checks obsolete.
+	// However, FCompilationExtensionManager starts a new compilation once the InputMappingContext's been loaded.
+	const bool bShouldEarlyOut = OldModifierCDO->InputMappingContext && OldModifierCDO->InputMappingContext->HasAnyFlags(RF_NeedLoad);
+	if (bShouldEarlyOut)
+	{
+		return;
+	}
+	
+	using namespace UE::VCamCoreEditor::Private;
+	ForEachNode(*CompilerContext->Blueprint, [this, CompilerContext, OldModifierCDO](UEdGraphNode& GraphNode)
+	{
+		if (UK2Node_EnhancedInputAction* EnhancedInputActionNode = Cast<UK2Node_EnhancedInputAction>(&GraphNode))
+		{
+			ValidateEnhancedInputAction(*CompilerContext, *OldModifierCDO, *EnhancedInputActionNode);
+		}
+
+		return EBreakBehaviour::Continue;
+	});
+}
+
+bool UModifierCompilationBlueprintExtension::RequiresRecompileToDetectIssues(UBlueprint& Blueprint)
+{
+	UVCamModifier* CDO = Blueprint.GeneratedClass ? Cast<UVCamModifier>(Blueprint.GeneratedClass->GetDefaultObject()) : nullptr;
+	if (!CDO)
+	{
+		return true;
+	}
+
+	bool bFoundAnyIssue = false;
+	using namespace UE::VCamCoreEditor::Private;
+	ForEachNode(Blueprint, [CDO, &bFoundAnyIssue](UEdGraphNode& GraphNode)
+	{
+		UK2Node_EnhancedInputAction* EnhancedInputActionNode = Cast<UK2Node_EnhancedInputAction>(&GraphNode);
+		bFoundAnyIssue = EnhancedInputActionNode && !DoesInputMappingContainReferencedAction(*CDO, *EnhancedInputActionNode);
+		return bFoundAnyIssue
+			? EBreakBehaviour::Break
+			: EBreakBehaviour::Continue;
+	});
+	return bFoundAnyIssue;
+}
+
+void UModifierCompilationBlueprintExtension::ValidateEnhancedInputAction(FKismetCompilerContext& CompilerContext, UVCamModifier& OldModifierCDO, UK2Node_EnhancedInputAction& EnhancedInputActionNode)
+{
+	if (!DoesInputMappingContainReferencedAction(OldModifierCDO, EnhancedInputActionNode))
+	{
+		const FName PropertyName = GET_MEMBER_NAME_CHECKED(UVCamModifier, InputMappingContext);
+		CompilerContext.MessageLog.Warning(TEXT("@@: InputAction "), &EnhancedInputActionNode)
+			->AddToken(FUObjectToken::Create(EnhancedInputActionNode.InputAction))
+			->AddToken(FTextToken::Create(LOCTEXT("EnhancedInput.MiddleOfMessage", " is not part of ")))
+			->AddToken(FUObjectToken::Create(OldModifierCDO.InputMappingContext))
+			->AddToken(FTextToken::Create(FText::Format(LOCTEXT("EnhancedInput.EndOfMessageFmt", " referenced by this modifier's {0} property"), FText::FromName(PropertyName))));
+		
+		EnhancedInputActionNode.bHasCompilerMessage = true;
+		EnhancedInputActionNode.ErrorMsg = FString::Printf(TEXT("InputAction %s is not part of %s referenced this modifier's %s property."),
+			EnhancedInputActionNode.InputAction ? *EnhancedInputActionNode.InputAction->GetName() : TEXT("None"),
+			OldModifierCDO.InputMappingContext ? *OldModifierCDO.InputMappingContext->GetName() : TEXT("None"),
+			*PropertyName.ToString()
+			);
+		EnhancedInputActionNode.ErrorType = EMessageSeverity::Warning;
+	}
+	else if (const bool bWasWarningGeneratedByUs = EnhancedInputActionNode.bHasCompilerMessage && EnhancedInputActionNode.ErrorMsg.Contains(TEXT("InputMappingContext")))
+	{
+		EnhancedInputActionNode.bHasCompilerMessage = false;
+		EnhancedInputActionNode.ErrorMsg.Reset();
+	}
+}
+
+bool UModifierCompilationBlueprintExtension::DoesInputMappingContainReferencedAction(UVCamModifier& OldModifierCDO,UK2Node_EnhancedInputAction& EnhancedInputActionNode)
+{
+	return OldModifierCDO.InputMappingContext && OldModifierCDO.InputMappingContext->GetMappings().ContainsByPredicate([SearchedAction = EnhancedInputActionNode.InputAction.Get()](const FEnhancedActionKeyMapping& Item)
+		{
+			return Item.Action == SearchedAction;
+		});
+}
+
+#undef LOCTEXT_NAMESPACE
diff --git a/Engine/Plugins/Experimental/VirtualCamera/Source/VCamCoreEditor/Private/Compilation/ModifierCompilationBlueprintExtension.h b/Engine/Plugins/Experimental/VirtualCamera/Source/VCamCoreEditor/Private/Compilation/ModifierCompilationBlueprintExtension.h
new file mode 100644
--- /dev/null
+++ b/Engine/Plugins/Experimental/VirtualCamera/Source/VCamCoreEditor/Private/Compilation/ModifierCompilationBlueprintExtension.h
@@ -0,0 +1,38 @@
+// Copyright Epic Games, Inc. All Rights Reserved.
+
+#pragma once
+
+#include "Blueprint/BlueprintExtension.h"
+#include "EdGraph/EdGraphNode.h"
+#include "Engine/Blueprint.h"
+#include "Containers/Array.h"
+#include "Containers/Set.h"
+#include "ModifierCompilationBlueprintExtension.generated.h"
+
+class UK2Node_EnhancedInputAction;
+class UVCamModifier;
+
+/**
+ * Validates modifier Blueprint when compiled.
+ * 
+ * Actions:
+ * - Add a warning to every input action node that is not found in the MappingContext in the class defaults.
+ */
+UCLASS()
+class VCAMCOREEDITOR_API UModifierCompilationBlueprintExtension : public UBlueprintExtension
+{
+	GENERATED_BODY()
+public:
+
+	//~ Begin UBlueprintExtension Interface
+	virtual void HandleGenerateFunctionGraphs(FKismetCompilerContext* CompilerContext) override;
+	//~ End UBlueprintExtension Interface
+
+	static bool RequiresRecompileToDetectIssues(UBlueprint& Blueprint);
+
+private:
+	
+	static void ValidateEnhancedInputAction(FKismetCompilerContext& CompilerContext, UVCamModifier& OldModifierCDO, UK2Node_EnhancedInputAction& EnhancedInputActionNode);
+	static bool DoesInputMappingContainReferencedAction(UVCamModifier& OldModifierCDO, UK2Node_EnhancedInputAction& EnhancedInputActionNode);
+};
+
diff --git a/Engine/Plugins/Experimental/VirtualCamera/Source/VCamCoreEditor/Private/VCamCoreEditorModule.cpp b/Engine/Plugins/Experimental/VirtualCamera/Source/VCamCoreEditor/Private/VCamCoreEditorModule.cpp
--- a/Engine/Plugins/Experimental/VirtualCamera/Source/VCamCoreEditor/Private/VCamCoreEditorModule.cpp
+++ b/Engine/Plugins/Experimental/VirtualCamera/Source/VCamCoreEditor/Private/VCamCoreEditorModule.cpp
@@ -30,6 +30,7 @@
 #include "Modules/ModuleInterface.h"
 #include "Modules/ModuleManager.h"
 #include "PropertyEditorModule.h"
+#include "Compilation/CompilationExtensionManager.h"
 
 #define LOCTEXT_NAMESPACE "FVCamCoreEditorModule"
 
@@ -53,10 +54,15 @@ namespace UE::VCamCoreEditor::Private
 		RegisterAutoGeneratedDefaultEvents();
 		RegisterCustomizations();
 		RegisterDefaultConnectionRemappingCustomizations();
+
+		CompilationExtensionManager = MakeShared<FCompilationExtensionManager>();
+		CompilationExtensionManager->Init();
 	}
 	
 	void FVCamCoreEditorModule::ShutdownModule()
 	{
+		CompilationExtensionManager.Reset();
+		
 		FKismetEditorUtilities::UnregisterAutoBlueprintNodeCreation(this);
 		if (UObjectInitialized() && !IsEngineExitRequested())
 		{
diff --git a/Engine/Plugins/Experimental/VirtualCamera/Source/VCamCoreEditor/Private/VCamCoreEditorModule.h b/Engine/Plugins/Experimental/VirtualCamera/Source/VCamCoreEditor/Private/VCamCoreEditorModule.h
--- a/Engine/Plugins/Experimental/VirtualCamera/Source/VCamCoreEditor/Private/VCamCoreEditorModule.h
+++ b/Engine/Plugins/Experimental/VirtualCamera/Source/VCamCoreEditor/Private/VCamCoreEditorModule.h
@@ -14,6 +14,8 @@ class UVCamWidget;
 
 namespace UE::VCamCoreEditor::Private
 {
+	class FCompilationExtensionManager;
+
 	class FVCamCoreEditorModule : public IVCamCoreEditorModule
 	{
 	public:
@@ -41,6 +43,8 @@ namespace UE::VCamCoreEditor::Private
 		/** Maps pointer to UVCamWidget classes */
 		TMap<TWeakObjectPtr<UClass>, FGetConnectionRemappingCustomization> ConnectionRemappingCustomizationFactories;
 
+		TSharedPtr<FCompilationExtensionManager> CompilationExtensionManager;
+
 		void RegisterAutoGeneratedDefaultEvents();
 		void RegisterCustomizations();
 		void RegisterDefaultConnectionRemappingCustomizations();
diff --git a/Engine/Plugins/Experimental/VirtualCamera/Source/VCamCoreEditor/VCamCoreEditor.build.cs b/Engine/Plugins/Experimental/VirtualCamera/Source/VCamCoreEditor/VCamCoreEditor.build.cs
--- a/Engine/Plugins/Experimental/VirtualCamera/Source/VCamCoreEditor/VCamCoreEditor.build.cs
+++ b/Engine/Plugins/Experimental/VirtualCamera/Source/VCamCoreEditor/VCamCoreEditor.build.cs
@@ -24,6 +24,7 @@ public class VCamCoreEditor : ModuleRules
 				"Engine",
 				"EditorFramework",
 				"EnhancedInput",
+				"InputBlueprintNodes",
 				"InputCore",
 				"LiveLinkInterface",
 				"PropertyEditor",
@@ -35,7 +36,7 @@ public class VCamCoreEditor : ModuleRules
 				"UnrealEd",
 				"VCamCore",
 				"VCamInput",
-				"VPUtilities"
+				"VPUtilities", 
 			}
 		);
 	}
diff --git a/Engine/Source/Runtime/Engine/Private/Blueprint.cpp b/Engine/Source/Runtime/Engine/Private/Blueprint.cpp
--- a/Engine/Source/Runtime/Engine/Private/Blueprint.cpp
+++ b/Engine/Source/Runtime/Engine/Private/Blueprint.cpp
@@ -625,6 +625,14 @@ UClass* UBlueprint::RegenerateClass(UClass* ClassToRegenerate, UObject* Previous
 	UBlueprint::ForceLoadMembers(this);
 
 PRAGMA_DISABLE_DEPRECATION_WARNINGS
+	for (auto It = Extensions.CreateIterator(); It; ++It)
+	{
+		if (!*It)
+		{
+			It.RemoveCurrent();
+		}
+	}
+	
 	for (UBlueprintExtension* Extension : Extensions)
 	{
 		ForceLoad(Extension);
