Map UVCamComponents to input devices

diff --git a/Engine/Plugins/Experimental/VirtualCamera/Source/VCamCore/Private/Input/InputVCamSubsystem.cpp b/Engine/Plugins/Experimental/VirtualCamera/Source/VCamCore/Private/Input/InputVCamSubsystem.cpp
new file mode 100644
--- /dev/null
+++ b/Engine/Plugins/Experimental/VirtualCamera/Source/VCamCore/Private/Input/InputVCamSubsystem.cpp
@@ -0,0 +1,177 @@
+// Copyright Epic Games, Inc. All Rights Reserved.
+
+#include "Input/InputVCamSubsystem.h"
+
+#include "Input/VCamPlayerInput.h"
+#include "LogVCamCore.h"
+#include "VCamInputProcessor.h"
+
+#include "Components/InputComponent.h"
+#include "Engine/World.h"
+#include "Framework/Application/SlateApplication.h"
+
+namespace UE::VCamCore::Private
+{
+	static EInputConsumptionRule DetermineInputConsumptionRule(UInputVCamSubsystem& Subsystem, bool bShouldGamepadConsumeAlways)
+	{
+		return Subsystem.GetWorld()->IsGameWorld()
+			// Normal gameplay behaviour
+		   ? EInputConsumptionRule::ConsumeIfUsed
+			// Editor should not consume all input ever. At most gamepads.
+		   : bShouldGamepadConsumeAlways ? EInputConsumptionRule::ConsumeOnlyGamepadIfUsed : EInputConsumptionRule::DoNotConsume;
+	}
+}
+
+void UInputVCamSubsystem::Initialize(FSubsystemCollectionBase& Collection)
+{
+	Super::Initialize(Collection);
+	UE_LOG(LogVCamCore, Log, TEXT("Initializing UInputVCamSubsystem..."));
+	
+	PlayerInput = NewObject<UVCamPlayerInput>(this);
+	
+	// Create and register the input preprocessor, this is what will call our "InputKey"
+	// function to drive input instead of a player controller
+	if (ensure(FSlateApplication::IsInitialized()))
+	{
+		// It's dangerous to consume input in editor (imagine typing something into search boxes but all L keys were consumed by VCam input)
+		// whereas probably expected by gameplay code.
+		using namespace UE::VCamCore::Private;
+		constexpr bool bDefaultShouldConsumeGamepads = false;
+		InputPreprocessor = MakeShared<FVCamInputProcessor>(*this, DetermineInputConsumptionRule(*this, bDefaultShouldConsumeGamepads));
+		FSlateApplication::Get().RegisterInputPreProcessor(InputPreprocessor, 0);	
+	}
+}
+
+void UInputVCamSubsystem::Deinitialize()
+{
+	Super::Deinitialize();
+	UE_LOG(LogVCamCore, Log, TEXT("De-initializing UInputVCamSubsystem..."));
+
+	if (FSlateApplication::IsInitialized())
+	{
+		FSlateApplication::Get().UnregisterInputPreProcessor(InputPreprocessor);
+	}
+
+	PlayerInput = nullptr;
+}
+
+bool UInputVCamSubsystem::InputKey(const FInputKeyParams& Params)
+{
+	return PlayerInput->InputKey(Params);
+}
+
+void UInputVCamSubsystem::PushInputComponent(UInputComponent* InInputComponent)
+{
+	if (!InInputComponent)
+	{
+		return;
+	}
+	
+	bool bPushed = false;
+	CurrentInputStack.RemoveSingle(InInputComponent);
+	for (int32 Index = CurrentInputStack.Num() - 1; Index >= 0; --Index)
+	{
+		UInputComponent* IC = CurrentInputStack[Index].Get();
+		if (IC == nullptr)
+		{
+			CurrentInputStack.RemoveAt(Index);
+		}
+		else if (IC->Priority <= InInputComponent->Priority)
+		{
+			CurrentInputStack.Insert(InInputComponent, Index + 1);
+			bPushed = true;
+			break;
+		}
+	}
+	if (!bPushed)
+	{
+		CurrentInputStack.Insert(InInputComponent, 0);
+		RequestRebuildControlMappings();
+	}
+}
+
+bool UInputVCamSubsystem::PopInputComponent(UInputComponent* InInputComponent)
+{
+	if (CurrentInputStack.RemoveSingle(InInputComponent) > 0)
+	{
+		InInputComponent->ClearBindingValues();
+		RequestRebuildControlMappings();
+		return true;
+	}
+	return false;
+}
+
+void UInputVCamSubsystem::SetShouldConsumeGamepadInput(EVCamGamepadInputMode GamepadInputMode)
+{
+	check(InputPreprocessor && PlayerInput);
+	
+	FVCamInputDeviceConfig InputSettings = PlayerInput->GetInputSettings();
+	InputSettings.GamepadInputMode = GamepadInputMode;
+	SetInputSettings(InputSettings);
+}
+
+const FVCamInputDeviceConfig& UInputVCamSubsystem::GetInputSettings() const
+{
+	// Undefined behaviour returning from dereferenced nullptr, let's make sure to assert.
+	UE_CLOG(PlayerInput == nullptr, LogVCamCore, Fatal, TEXT("PlayerInput is designed to exist for the lifetime of UInputVCamSubsystem. Investigate!"));
+	return PlayerInput->GetInputSettings();
+}
+
+void UInputVCamSubsystem::SetInputSettings(const FVCamInputDeviceConfig& Input)
+{
+	check(PlayerInput);
+	PlayerInput->SetInputSettings(Input);
+
+	const bool bShouldConsumeGamepad = Input.GamepadInputMode == EVCamGamepadInputMode::IgnoreAndConsume;
+	InputPreprocessor->SetInputConsumptionRule(
+		UE::VCamCore::Private::DetermineInputConsumptionRule(*this, bShouldConsumeGamepad)
+		);
+}
+
+UWorld* UInputVCamSubsystem::GetTickableGameObjectWorld() const
+{
+	return GetTypedOuter<UWorld>();
+}
+
+ETickableTickType UInputVCamSubsystem::GetTickableTickType() const
+{
+	return IsTemplate() ? ETickableTickType::Never : ETickableTickType::Conditional;
+}
+
+bool UInputVCamSubsystem::IsAllowedToTick() const
+{
+	return PlayerInput != nullptr;
+}
+
+void UInputVCamSubsystem::Tick(float DeltaTime)
+{
+	if (!ensure(PlayerInput))
+	{
+		return;
+	}
+
+	FModifyContextOptions Options;
+	Options.bForceImmediately = true;
+	RequestRebuildControlMappings(Options);
+
+	TArray<UInputComponent*> InputStack;
+	for (auto It = CurrentInputStack.CreateIterator(); It; ++It)
+	{
+		if (UInputComponent* InputComponent = It->Get())
+		{
+			InputStack.Push(InputComponent);
+		}
+		else
+		{
+			It.RemoveCurrent();
+		}
+	}
+	
+	PlayerInput->Tick(DeltaTime);
+	PlayerInput->ProcessInputStack(InputStack, DeltaTime, false);
+}
+
+UEnhancedPlayerInput* UInputVCamSubsystem::GetPlayerInput() const
+{
+	return PlayerInput;
+}
diff --git a/Engine/Plugins/Experimental/VirtualCamera/Source/VCamCore/Private/Input/VCamInputProcessor.cpp b/Engine/Plugins/Experimental/VirtualCamera/Source/VCamCore/Private/Input/VCamInputProcessor.cpp
new file mode 100644
--- /dev/null
+++ b/Engine/Plugins/Experimental/VirtualCamera/Source/VCamCore/Private/Input/VCamInputProcessor.cpp
@@ -0,0 +1,184 @@
+// Copyright Epic Games, Inc. All Rights Reserved.
+
+#include "VCamInputProcessor.h"
+
+#include "Input/InputVCamSubsystem.h"
+
+#include "Framework/Application/SlateApplication.h"
+#include "GameFramework/PlayerInput.h"
+#include "Input/Events.h"
+
+namespace UE::VCamCore::Private
+{
+	FVCamInputProcessor::FVCamInputProcessor(UInputVCamSubsystem& OwningSubsystem, EInputConsumptionRule InputConsumptionRule)
+		: OwningSubsystem(&OwningSubsystem)
+		, InputConsumptionRule(InputConsumptionRule)
+	{}
+
+	void FVCamInputProcessor::Tick(const float DeltaTime, FSlateApplication& SlateApp, TSharedRef<ICursor> Cursor)
+	{
+		ProcessAccumulatedPointerInput(DeltaTime);
+	}
+
+	bool FVCamInputProcessor::HandleKeyDownEvent(FSlateApplication& SlateApp, const FKeyEvent& InKeyEvent)
+	{
+		FInputKeyParams Params = {};
+		Params.Key = InKeyEvent.GetKey();
+		Params.Event = IE_Pressed;
+		Params.Delta.X = 1.0;
+		Params.DeltaTime = SlateApp.GetDeltaTime();
+		Params.NumSamples = Params.Key.IsAnalog() ? 1 : 0;
+		Params.InputDevice = InKeyEvent.GetInputDeviceId();
+
+		return InputKeyToSubsystem(Params);
+	}
+
+	bool FVCamInputProcessor::HandleKeyUpEvent(FSlateApplication& SlateApp, const FKeyEvent& InKeyEvent)
+	{
+		FInputKeyParams Params = {};
+		Params.Key = InKeyEvent.GetKey();
+		Params.Event = IE_Released;
+		Params.Delta.X = 0.0;
+		Params.DeltaTime = SlateApp.GetDeltaTime();
+		Params.NumSamples = Params.Key.IsAnalog() ? 1 : 0;
+		Params.InputDevice = InKeyEvent.GetInputDeviceId();
+		
+		return InputKeyToSubsystem(Params);
+	}
+
+	bool FVCamInputProcessor::HandleAnalogInputEvent(FSlateApplication& SlateApp, const FAnalogInputEvent& InAnalogInputEvent)
+	{
+		FInputKeyParams Params = {};
+		Params.Key = InAnalogInputEvent.GetKey();
+		Params.Event = IE_Pressed;
+		Params.Delta.X = InAnalogInputEvent.GetAnalogValue();
+		Params.DeltaTime = SlateApp.GetDeltaTime();
+		Params.NumSamples = 1;
+		Params.InputDevice = InAnalogInputEvent.GetInputDeviceId();
+		
+		return InputKeyToSubsystem(Params);
+	}
+
+	bool FVCamInputProcessor::HandleMouseMoveEvent(FSlateApplication& SlateApp, const FPointerEvent& MouseEvent)
+	{
+		UpdateCachedPointerPosition(MouseEvent);
+		return IInputProcessor::HandleMouseMoveEvent(SlateApp, MouseEvent);
+	}
+
+	bool FVCamInputProcessor::HandleMouseButtonDownEvent(FSlateApplication& SlateApp, const FPointerEvent& MouseEvent)
+	{
+		FInputKeyParams Params = {};
+		Params.Key = MouseEvent.GetEffectingButton();
+		Params.Event = IE_Pressed;
+		Params.Delta.X = 1.0;
+		Params.DeltaTime = SlateApp.GetDeltaTime();
+		Params.NumSamples = 0;
+
+		return InputKeyToSubsystem(Params);
+	}
+
+	bool FVCamInputProcessor::HandleMouseButtonUpEvent(FSlateApplication& SlateApp, const FPointerEvent& MouseEvent)
+	{
+		FInputKeyParams Params = {};
+		Params.Key = MouseEvent.GetEffectingButton();
+		Params.Event = IE_Released;
+		Params.Delta.X = 0.0;
+		Params.DeltaTime = SlateApp.GetDeltaTime();
+		Params.NumSamples = 0;
+		
+		return InputKeyToSubsystem(Params);
+	}
+
+	bool FVCamInputProcessor::HandleMouseButtonDoubleClickEvent(FSlateApplication& SlateApp, const FPointerEvent& MouseEvent)
+	{
+		FInputKeyParams Params = {};
+		Params.Key = MouseEvent.GetEffectingButton();
+		Params.Event = IE_DoubleClick;
+		Params.Delta.X = 1.0;
+		Params.DeltaTime = SlateApp.GetDeltaTime();
+		Params.NumSamples = 0;
+		
+		return InputKeyToSubsystem(Params);
+	}
+
+	bool FVCamInputProcessor::HandleMouseWheelOrGestureEvent(FSlateApplication& SlateApp, const FPointerEvent& InWheelEvent, const FPointerEvent* InGestureEvent)
+	{
+		const FKey MouseWheelKey = InWheelEvent.GetWheelDelta() < 0 ? EKeys::MouseScrollDown : EKeys::MouseScrollUp;
+
+		// Input the Mouse wheel key events (mouse scroll down or scroll up) as being pressed and released this frame
+		// The SceneViewport Inputs the Mouse Scroll Wheel buttons up and down in the same frame, this replicates that behavior
+		{
+			FInputKeyParams PressedParams = {};
+			PressedParams.Key = MouseWheelKey;
+			PressedParams.Event = IE_Pressed;
+			PressedParams.Delta.X = 1.0;
+			PressedParams.DeltaTime = SlateApp.GetDeltaTime();
+			PressedParams.NumSamples = 0;
+			
+			FInputKeyParams ReleasedParams = PressedParams;
+			ReleasedParams.Event = IE_Released;
+
+			InputKeyToSubsystem(PressedParams);
+			InputKeyToSubsystem(ReleasedParams);	
+		}
+		// Input the wheel axis delta to get the MouseWheelAxis button working
+		{
+			FInputKeyParams Params = {};
+			Params.Key = EKeys::MouseWheelAxis;
+			Params.Delta.X = InWheelEvent.GetWheelDelta();
+			Params.DeltaTime = SlateApp.GetDeltaTime();
+			Params.NumSamples = 1;
+			
+			InputKeyToSubsystem(Params);
+		}
+		
+		return IInputProcessor::HandleMouseWheelOrGestureEvent(SlateApp, InWheelEvent, InGestureEvent);
+	}
+
+	void FVCamInputProcessor::UpdateCachedPointerPosition(const FPointerEvent& MouseEvent)
+	{
+		CachedCursorDelta = MouseEvent.GetCursorDelta();
+		
+		++NumCursorSamplesThisFrame.X;
+		++NumCursorSamplesThisFrame.Y;
+	}
+
+	void FVCamInputProcessor::ProcessAccumulatedPointerInput(float DeltaTime)
+	{
+		// Input the MouseX value
+		{
+			FInputKeyParams Params = {};
+			Params.Key = EKeys::MouseX;
+			Params.Delta.X = CachedCursorDelta.X;
+			Params.DeltaTime = DeltaTime;
+			Params.NumSamples = NumCursorSamplesThisFrame.X;
+			InputKeyToSubsystem(Params);
+		}
+
+		// Input the MouseY value
+		{
+			FInputKeyParams Params = {};
+			Params.Key = EKeys::MouseY;
+			Params.Delta.X = CachedCursorDelta.Y;
+			Params.DeltaTime = DeltaTime;
+			Params.NumSamples = NumCursorSamplesThisFrame.Y;
+			InputKeyToSubsystem(Params);
+		}
+		
+		NumCursorSamplesThisFrame = FIntPoint::ZeroValue;
+		CachedCursorDelta = FVector2D::ZeroVector;
+	}
+
+	bool FVCamInputProcessor::InputKeyToSubsystem(const FInputKeyParams& Params)
+	{
+		if (ensure(OwningSubsystem.IsValid()))
+		{
+			const bool bCanConsumeInput = InputConsumptionRule == EInputConsumptionRule::ConsumeIfUsed
+				|| (InputConsumptionRule == EInputConsumptionRule::ConsumeOnlyGamepadIfUsed && Params.IsGamepad());
+			const bool bWantsToConsumeInput = OwningSubsystem->InputKey(Params);
+			return bCanConsumeInput && bWantsToConsumeInput;
+		}
+		
+		return false;
+	}
+}
diff --git a/Engine/Plugins/Experimental/VirtualCamera/Source/VCamCore/Private/Input/VCamInputProcessor.h b/Engine/Plugins/Experimental/VirtualCamera/Source/VCamCore/Private/Input/VCamInputProcessor.h
new file mode 100644
--- /dev/null
+++ b/Engine/Plugins/Experimental/VirtualCamera/Source/VCamCore/Private/Input/VCamInputProcessor.h
@@ -0,0 +1,76 @@
+// Copyright Epic Games, Inc. All Rights Reserved.
+
+#pragma once
+
+#include "Framework/Application/IInputProcessor.h"
+#include "UObject/WeakObjectPtr.h"
+#include "UObject/WeakObjectPtrTemplates.h"
+
+class UInputVCamSubsystem;
+struct FInputKeyParams;
+
+namespace UE::VCamCore::Private
+{
+	enum class EInputConsumptionRule
+	{
+		/** Never consume input. */
+		DoNotConsume,
+		/** Consume input if InputKey returned true */
+		ConsumeIfUsed,
+		/** Consume input if InputKey returned true but only for gamepads. For backwards compatibility. */
+		ConsumeOnlyGamepadIfUsed
+	};
+	
+	/**
+	 * Passes inputs to UVCamPlayerInput::InputKey.
+	 * 
+	 * This code is a modified version of FEnhancedInputEditorProcessor; the only difference is where the input ends up
+	 * getting sent to.
+	 */
+	class FVCamInputProcessor : public IInputProcessor
+	{
+	public:
+
+		FVCamInputProcessor(UInputVCamSubsystem& OwningSubsystem, EInputConsumptionRule InputConsumptionRule);
+
+		void SetInputConsumptionRule(EInputConsumptionRule NewInputConsumptionRule) { InputConsumptionRule = NewInputConsumptionRule; };
+		
+		//~ Begin IInputProcessor interface
+		virtual void Tick(const float DeltaTime, FSlateApplication& SlateApp, TSharedRef<ICursor> Cursor) override;
+		virtual bool HandleKeyDownEvent(FSlateApplication& SlateApp, const FKeyEvent& InKeyEvent) override;
+		virtual bool HandleKeyUpEvent(FSlateApplication& SlateApp, const FKeyEvent& InKeyEvent) override;
+		virtual bool HandleAnalogInputEvent(FSlateApplication& SlateApp, const FAnalogInputEvent& InAnalogInputEvent) override;
+		virtual bool HandleMouseMoveEvent(FSlateApplication& SlateApp, const FPointerEvent& MouseEvent) override;
+		virtual bool HandleMouseButtonDownEvent( FSlateApplication& SlateApp, const FPointerEvent& MouseEvent) override;
+		virtual bool HandleMouseButtonUpEvent( FSlateApplication& SlateApp, const FPointerEvent& MouseEvent) override;
+		virtual bool HandleMouseButtonDoubleClickEvent(FSlateApplication& SlateApp, const FPointerEvent& MouseEvent) override;
+		virtual bool HandleMouseWheelOrGestureEvent(FSlateApplication& SlateApp, const FPointerEvent& InWheelEvent, const FPointerEvent* InGestureEvent) override;
+		virtual const TCHAR* GetDebugName() const { return TEXT("FVCamInputProcessor"); }
+		//~ End IInputProcessor interface
+
+	private:
+
+		TWeakObjectPtr<UInputVCamSubsystem> OwningSubsystem;
+		
+		/** Determines whether input can be consumed when InputKey returns true. */
+		EInputConsumptionRule InputConsumptionRule;
+
+		/**
+		 * The number of cursor samples that happened this frame. The X and Y will be incremented when there is a mouse move event this frame
+		 * and it will be reset on ProcessAccumulatedPointerInput.
+		 */
+		FIntPoint NumCursorSamplesThisFrame = FIntPoint::ZeroValue;
+	
+		/** The value of the mouse cursor from the most recent mouse event */
+		FVector2D CachedCursorDelta = FVector2D::ZeroVector;
+
+		void UpdateCachedPointerPosition(const FPointerEvent& MouseEvent);
+		/** Called on Tick to process any accumulated mouse input */
+		void ProcessAccumulatedPointerInput(float DeltaTime);
+	
+		/** Tells the Enhanced Input Editor subsystem about the given key event */
+		bool InputKeyToSubsystem(const FInputKeyParams& Params);
+	};
+}
+
+
diff --git a/Engine/Plugins/Experimental/VirtualCamera/Source/VCamCore/Private/Input/VCamPlayerInput.cpp b/Engine/Plugins/Experimental/VirtualCamera/Source/VCamCore/Private/Input/VCamPlayerInput.cpp
new file mode 100644
--- /dev/null
+++ b/Engine/Plugins/Experimental/VirtualCamera/Source/VCamCore/Private/Input/VCamPlayerInput.cpp
@@ -0,0 +1,94 @@
+// Copyright Epic Games, Inc. All Rights Reserved.
+
+#include "Input/VCamPlayerInput.h"
+
+#include "LogVCamCore.h"
+
+namespace UE::VCamCore::Private
+{
+	static FString LexInputEvent(EInputEvent InputEvent)
+	{
+		switch (InputEvent)
+		{
+		case IE_Pressed: return TEXT("pressed");
+		case IE_Released: return TEXT("released");
+		case IE_Repeat: return TEXT("repeat");
+		case IE_DoubleClick: return TEXT("double-click");
+		case IE_Axis: return TEXT("axis");
+		case IE_MAX: // pass-through
+		default: return TEXT("invalid");
+		}
+	}
+	
+	static FString ToString(const FInputKeyParams& Params)
+	{
+		return FString::Printf(TEXT("{ InputID: %d, Key: %s, EInputEvent: %s, bIsGamepad: %s }"),
+			Params.InputDevice.GetId(),
+			*Params.Key.ToString(),
+			*LexInputEvent(Params.Event),
+			Params.IsGamepad() ? TEXT("true") : TEXT("false")
+			);
+	}
+
+	static void LogInput(FVCamInputDeviceConfig InputDeviceSettings, const FInputKeyParams& Params, bool bIsFilteredOut)
+	{
+		switch (InputDeviceSettings.LoggingMode)
+		{
+		case EVCamInputLoggingMode::OnlyConsumable:
+			if (!bIsFilteredOut)
+			{
+				UE_LOG(LogVCamInputDebug, Log, TEXT("%s"), *ToString(Params));
+			}
+			break;
+		case EVCamInputLoggingMode::OnlyGamepad:
+			if (Params.IsGamepad())
+			{
+				UE_LOG(LogVCamInputDebug, Log, TEXT("%s"), *ToString(Params));
+			}
+			break;
+		case EVCamInputLoggingMode::AllExceptMouse:
+			if (!Params.Key.IsMouseButton())
+			{
+				UE_LOG(LogVCamInputDebug, Log, TEXT("%s"), *ToString(Params));
+			}
+			break;
+		
+		case EVCamInputLoggingMode::All:
+			UE_LOG(LogVCamInputDebug, Log, TEXT("%s"), *ToString(Params));
+		default: ;
+		}
+	}
+}
+
+bool UVCamPlayerInput::InputKey(const FInputKeyParams& Params)
+{
+	const bool bIsKeyboard = !Params.IsGamepad() && !Params.Key.IsAnalog() && !Params.Key.IsMouseButton() && !Params.Key.IsTouch()
+		// Keyboard is always mapped to 0
+		&& Params.InputDevice.GetId() == 0;
+	const bool bCanCheckAllowList = Params.InputDevice != INPUTDEVICEID_NONE && !bIsKeyboard;
+	
+	const bool bSkipGamepad = Params.IsGamepad() && (InputDeviceSettings.GamepadInputMode != EVCamGamepadInputMode::Allow);
+	const bool bSkipMouse = InputDeviceSettings.MouseInputMode == EVCamInputMode::Ignore && Params.Key.IsMouseButton();
+	const bool bSkipKeyboard = InputDeviceSettings.KeyboardInputMode == EVCamInputMode::Ignore && bIsKeyboard;
+	const bool bSkipNonAllowListed = bCanCheckAllowList && !InputDeviceSettings.AllowedInputDeviceIds.Contains(Params.InputDevice.GetId());
+	const bool bIsFilteredOut = bSkipGamepad || bSkipMouse || bSkipKeyboard || bSkipNonAllowListed;
+	
+	UE::VCamCore::Private::LogInput(InputDeviceSettings, Params, bIsFilteredOut);
+	if (!bIsFilteredOut)
+	{
+		return Super::InputKey(Params);
+	}
+
+	const bool bSkippedGamepad = bSkipGamepad && InputDeviceSettings.GamepadInputMode == EVCamGamepadInputMode::IgnoreAndConsume;
+	return bSkippedGamepad;
+}
+
+void UVCamPlayerInput::ProcessInputStack(const TArray<UInputComponent*>& InputComponentStack, const float DeltaTime, const bool bGamePaused)
+{
+	Super::ProcessInputStack(InputComponentStack, DeltaTime, bGamePaused);
+}
+
+void UVCamPlayerInput::SetInputSettings(const FVCamInputDeviceConfig& Input)
+{
+	InputDeviceSettings = Input;
+}
diff --git a/Engine/Plugins/Experimental/VirtualCamera/Source/VCamCore/Private/LogVCamCore.cpp b/Engine/Plugins/Experimental/VirtualCamera/Source/VCamCore/Private/LogVCamCore.cpp
--- a/Engine/Plugins/Experimental/VirtualCamera/Source/VCamCore/Private/LogVCamCore.cpp
+++ b/Engine/Plugins/Experimental/VirtualCamera/Source/VCamCore/Private/LogVCamCore.cpp
@@ -2,4 +2,5 @@
 
 #include "LogVCamCore.h"
 
-DEFINE_LOG_CATEGORY(LogVCamCore);
\ No newline at end of file
+DEFINE_LOG_CATEGORY(LogVCamCore);
+DEFINE_LOG_CATEGORY(LogVCamInputDebug);
\ No newline at end of file
diff --git a/Engine/Plugins/Experimental/VirtualCamera/Source/VCamCore/Private/LogVCamCore.h b/Engine/Plugins/Experimental/VirtualCamera/Source/VCamCore/Private/LogVCamCore.h
--- a/Engine/Plugins/Experimental/VirtualCamera/Source/VCamCore/Private/LogVCamCore.h
+++ b/Engine/Plugins/Experimental/VirtualCamera/Source/VCamCore/Private/LogVCamCore.h
@@ -2,6 +2,7 @@
 
 #pragma once
 
-#include "CoreMinimal.h"
+#include "Logging/LogMacros.h"
 
-DECLARE_LOG_CATEGORY_EXTERN(LogVCamCore, Log, All);
\ No newline at end of file
+DECLARE_LOG_CATEGORY_EXTERN(LogVCamCore, Log, All);
+DECLARE_LOG_CATEGORY_EXTERN(LogVCamInputDebug, Log, All);
\ No newline at end of file
diff --git a/Engine/Plugins/Experimental/VirtualCamera/Source/VCamCore/Private/Modifier/VCamModifier.cpp b/Engine/Plugins/Experimental/VirtualCamera/Source/VCamCore/Private/Modifier/VCamModifier.cpp
--- a/Engine/Plugins/Experimental/VirtualCamera/Source/VCamCore/Private/Modifier/VCamModifier.cpp
+++ b/Engine/Plugins/Experimental/VirtualCamera/Source/VCamCore/Private/Modifier/VCamModifier.cpp
@@ -54,7 +54,7 @@ void UVCamModifier::Initialize(UVCamModifierContext* Context, UInputComponent* I
 		UInputDelegateBinding::BindInputDelegates(GetClass(), InputComponent, this);
 	}
 
-	bRequiresInitialization = false;
+	bIsInitialized = false;
 }
 
 void UVCamModifier::Deinitialize()
@@ -65,12 +65,12 @@ void UVCamModifier::Deinitialize()
 		VCamComponent->UnregisterObjectForInput(this);
 	}
 
-	bRequiresInitialization = true;
+	bIsInitialized = true;
 }
 
 void UVCamModifier::BeginDestroy()
 {
-	if (!bRequiresInitialization)
+	if (!bIsInitialized)
 	{
 		Deinitialize();
 	}
@@ -82,7 +82,7 @@ void UVCamModifier::PostLoad()
 {
 	Super::PostLoad();
 
-	bRequiresInitialization = true;
+	bIsInitialized = true;
 }
 
 UVCamComponent* UVCamModifier::GetOwningVCamComponent() const
diff --git a/Engine/Plugins/Experimental/VirtualCamera/Source/VCamCore/Private/VCamComponent.cpp b/Engine/Plugins/Experimental/VirtualCamera/Source/VCamCore/Private/VCamComponent.cpp
--- a/Engine/Plugins/Experimental/VirtualCamera/Source/VCamCore/Private/VCamComponent.cpp
+++ b/Engine/Plugins/Experimental/VirtualCamera/Source/VCamCore/Private/VCamComponent.cpp
@@ -19,6 +19,7 @@
 #include "ILiveLinkClient.h"
 #include "InputMappingContext.h"
 #include "PlayerMappableKeySettings.h"
+#include "Input/InputVCamSubsystem.h"
 #include "Roles/LiveLinkCameraRole.h"
 #include "Roles/LiveLinkTransformRole.h"
 
@@ -99,27 +100,7 @@ UVCamComponent::UVCamComponent()
 
 void UVCamComponent::OnComponentDestroyed(bool bDestroyingHierarchy)
 {
-	bLockViewportToCamera_DEPRECATED = false;
-	UnlockAllViewports();
-
-	for (UVCamOutputProviderBase* Provider : OutputProviders)
-	{
-		if (Provider && Provider->IsInitialized())
-		{
-			Provider->Deinitialize();
-		}
-	}
-
-	for (FModifierStackEntry& ModifierEntry : ModifierStack)
-	{
-		TObjectPtr<UVCamModifier>& Modifier = ModifierEntry.GeneratedModifier;
-		if (Modifier && !Modifier->IsInitialized())
-		{
-			Modifier->Deinitialize();
-		}
-	}
-
-	UnregisterInputComponent();
+	Deinitialize();
 
 #if WITH_EDITOR
 	// Remove all event listeners
@@ -190,98 +171,43 @@ void UVCamComponent::NotifyComponentWasReplaced(UVCamComponent* ReplacementCompo
 	DestroyComponent();
 }
 
-IEnhancedInputSubsystemInterface* UVCamComponent::GetEnhancedInputSubsystemInterface() const
-{
-	if (const UWorld* World = GetWorld(); IsValid(World) && World->IsGameWorld())
-	{
-		
-		if (const ULocalPlayer* FirstLocalPlayer = World->GetFirstLocalPlayerFromController())
-		{
-			return ULocalPlayer::GetSubsystem<UEnhancedInputLocalPlayerSubsystem>(FirstLocalPlayer);
-		}
-	}
-#if WITH_EDITOR
-	else if (GEditor)
-	{
-		return GEditor->GetEditorSubsystem<UEnhancedInputEditorSubsystem>();
-	}
-#endif
-	return nullptr;
-}
-
 void UVCamComponent::RegisterInputComponent()
 {
 	// Ensure we start from a clean slate
 	UnregisterInputComponent();
 	
-	if (const UWorld* World = GetWorld(); IsValid(World) && World->IsGameWorld())
-	{
-		if (APlayerController* PC = World->GetFirstPlayerController())
-		{
-			PC->PushInputComponent(InputComponent);
-			bIsInputRegistered = true;
-		}
-	}
-#if WITH_EDITOR
-	else if (GEditor)
+	if (UInputVCamSubsystem* InputSubsystem = GetInputVCamSubsystem())
 	{
-		if (UEnhancedInputEditorSubsystem* EditorInputSubsystem = GEditor->GetEditorSubsystem<UEnhancedInputEditorSubsystem>())
-		{
-			EditorInputSubsystem->PushInputComponent(InputComponent);
-			EditorInputSubsystem->StartConsumingInput();
-			bIsInputRegistered = true;
-		}
+		InputSubsystem->PushInputComponent(InputComponent);
 	}
-#endif
 }
 
 void UVCamComponent::UnregisterInputComponent()
 {
-	// Removes the component from both editor and runtime input systems if possible
-	//
-	// Note: Despite the functions being called "Pop" it's actually just removing our specific input component from
-	// the stack rather than blindly popping the top component
-	if (const UWorld* World = GetWorld(); IsValid(World) && World->IsGameWorld())
+	if (UInputVCamSubsystem* InputSubsystem = GetInputVCamSubsystem())
 	{
-		if (APlayerController* PC = World->GetFirstPlayerController())
-		{
-			PC->PopInputComponent(InputComponent);
-		}
-	}
-#if WITH_EDITOR
-	if (GEditor)
-	{
-		if (UEnhancedInputEditorSubsystem* EditorInputSubsystem = GEditor->GetEditorSubsystem<UEnhancedInputEditorSubsystem>())
-		{
-			EditorInputSubsystem->PopInputComponent(InputComponent);
-		}
+		// Note: Despite the functions being called "Pop" it's actually just removing our specific input component from
+		// the stack rather than blindly popping the top component
+		InputSubsystem->PopInputComponent(InputComponent);
 	}
-#endif
+	
 	AppliedInputContexts.Reset();
-	bIsInputRegistered = false;
 }
 
 bool UVCamComponent::CanUpdate() const
 {
+	const bool bShouldUpdate = bEnabled && !bIsEditorObjectButPIEIsRunning;
+	
+	// This prevents us updating in asset editors or invalid worlds
+	constexpr int ValidWorldTypes = EWorldType::Game | EWorldType::PIE | EWorldType::Editor;
 	UWorld* World = GetWorld();
-	if (bEnabled && IsValid(this) && !bIsEditorObjectButPIEIsRunning && World)
-	{
-		// Check that we only update in a valid world type
-		// This prevents us updating in asset editors or invalid worlds
-		constexpr int ValidWorldTypes = EWorldType::Game | EWorldType::PIE | EWorldType::Editor;
-		if ((World->WorldType & ValidWorldTypes) != EWorldType::None)
-		{
-			if (const USceneComponent* ParentComponent = GetAttachParent())
-			{
-				if (ParentComponent->IsA<UCineCameraComponent>())
-				{
-					// Component is valid to use if it is enabled, has a parent and that parent is a CineCamera derived component
-					return true;
-				}
-			}
-		}
-	}
-	return false;
+	const bool bIsSupportedWorld = World && (World->WorldType & ValidWorldTypes) != EWorldType::None;
+
+	// Modifiers requires a cine camera component
+	const USceneComponent* ParentComponent = GetAttachParent();
+	const bool bHasValidComponent = ParentComponent->IsA<UCineCameraComponent>();
+	
+	return bShouldUpdate && bIsSupportedWorld && bHasValidComponent;
 
 }
 
@@ -398,12 +324,9 @@ void UVCamComponent::PostEditChangeProperty(FPropertyChangedEvent& PropertyChang
 	FProperty* Property = PropertyChangedEvent.MemberProperty;
 	if (Property && PropertyChangedEvent.ChangeType != EPropertyChangeType::Interactive)
 	{
-		static FName NAME_Enabled = GET_MEMBER_NAME_CHECKED(UVCamComponent, bEnabled);
-		static FName NAME_ModifierStack = GET_MEMBER_NAME_CHECKED(UVCamComponent, ModifierStack);
-		static FName NAME_InputProfile = GET_MEMBER_NAME_CHECKED(UVCamComponent, InputProfile);
 		const FName PropertyName = Property->GetFName();
-
-		if (PropertyName == NAME_Enabled)
+		
+		if (PropertyName == GET_MEMBER_NAME_CHECKED(UVCamComponent, bEnabled))
 		{
 			// Only act here if we are a struct (like FModifierStackEntry)
 			if (!Property->GetOwner<UClass>())
@@ -411,15 +334,19 @@ void UVCamComponent::PostEditChangeProperty(FPropertyChangedEvent& PropertyChang
 				SetEnabled(bEnabled);
 			}
 		}
-		else if (PropertyName == NAME_ModifierStack)
+		else if (PropertyName == GET_MEMBER_NAME_CHECKED(UVCamComponent, ModifierStack))
 		{
 			ValidateModifierStack();
 			SavedModifierStack.Empty();
 		}
-		else if (PropertyName == NAME_InputProfile)
+		else if (PropertyName == GET_MEMBER_NAME_CHECKED(UVCamComponent, InputProfile))
 		{
 			ApplyInputProfile();
 		}
+		else if (PropertyName == GET_MEMBER_NAME_CHECKED(UVCamComponent, InputDeviceSettings))
+		{
+			SetInputDeviceSettings(InputDeviceSettings);
+		}
 
 		for (const UVCamOutputProviderBase* OutputProvider : OutputProviders)
 		{
@@ -518,7 +445,7 @@ void UVCamComponent::OnTargetViewportEdited()
 
 void UVCamComponent::AddInputMappingContext(const UVCamModifier* Modifier)
 {
-	if (IEnhancedInputSubsystemInterface* EnhancedInputSubsystemInterface = GetEnhancedInputSubsystemInterface())
+	if (IEnhancedInputSubsystemInterface* EnhancedInputSubsystemInterface = GetInputVCamSubsystem())
 	{
 		const int32 InputPriority = Modifier->InputContextPriority;
 		const UInputMappingContext* IMC = Modifier->InputMappingContext;
@@ -536,7 +463,7 @@ void UVCamComponent::AddInputMappingContext(const UVCamModifier* Modifier)
 
 void UVCamComponent::RemoveInputMappingContext(const UVCamModifier* Modifier)
 {
-	if (IEnhancedInputSubsystemInterface* EnhancedInputSubsystemInterface = GetEnhancedInputSubsystemInterface())
+	if (IEnhancedInputSubsystemInterface* EnhancedInputSubsystemInterface = GetInputVCamSubsystem())
 	{
 		const UInputMappingContext* IMC = Modifier->InputMappingContext;
 		if (IsValid(IMC))
@@ -549,7 +476,7 @@ void UVCamComponent::RemoveInputMappingContext(const UVCamModifier* Modifier)
 
 void UVCamComponent::AddInputMappingContext(UInputMappingContext* Context, int32 Priority)
 {
-	if (IEnhancedInputSubsystemInterface* EnhancedInputSubsystemInterface = GetEnhancedInputSubsystemInterface())
+	if (IEnhancedInputSubsystemInterface* EnhancedInputSubsystemInterface = GetInputVCamSubsystem())
 	{
 		if (IsValid(Context))
 		{
@@ -561,7 +488,7 @@ void UVCamComponent::AddInputMappingContext(UInputMappingContext* Context, int32
 
 void UVCamComponent::RemoveInputMappingContext(UInputMappingContext* Context)
 {
-	if (IEnhancedInputSubsystemInterface* EnhancedInputSubsystemInterface = GetEnhancedInputSubsystemInterface())
+	if (IEnhancedInputSubsystemInterface* EnhancedInputSubsystemInterface = GetInputVCamSubsystem())
 	{
 		if (IsValid(Context))
 		{
@@ -609,7 +536,7 @@ bool UVCamComponent::AddInputProfileWithCurrentlyActiveMappings(const FName Prof
 		TargetInputProfile = &VCamInputSettings->InputProfiles.Add(ProfileName);
 	}
 
-	if (const IEnhancedInputSubsystemInterface* EnhancedInputSubsystemInterface = GetEnhancedInputSubsystemInterface())
+	if (const IEnhancedInputSubsystemInterface* EnhancedInputSubsystemInterface = GetInputVCamSubsystem())
 	{
 		TArray<FEnhancedActionKeyMapping> PlayerMappableActionKeyMappings;
 		for (const UInputMappingContext* MappingContext : AppliedInputContexts)
@@ -663,7 +590,7 @@ bool UVCamComponent::SaveCurrentInputProfileToSettings(const FName ProfileName)
 
 TArray<FEnhancedActionKeyMapping> UVCamComponent::GetAllPlayerMappableActionKeyMappings() const
 {
-	if (const IEnhancedInputSubsystemInterface* EnhancedInputSubsystemInterface = GetEnhancedInputSubsystemInterface())
+	if (const IEnhancedInputSubsystemInterface* EnhancedInputSubsystemInterface = GetInputVCamSubsystem())
 	{
 		return EnhancedInputSubsystemInterface->GetAllPlayerMappableActionKeyMappings();
 	}
@@ -672,14 +599,13 @@ TArray<FEnhancedActionKeyMapping> UVCamComponent::GetAllPlayerMappableActionKeyM
 
 FKey UVCamComponent::GetPlayerMappedKey(const FName MappingName) const
 {
-	if (const IEnhancedInputSubsystemInterface* EnhancedInputSubsystemInterface = GetEnhancedInputSubsystemInterface())
+	if (const IEnhancedInputSubsystemInterface* EnhancedInputSubsystemInterface = GetInputVCamSubsystem())
 	{
 		return EnhancedInputSubsystemInterface->GetPlayerMappedKeyInSlot(MappingName);
 	}
 	return EKeys::Invalid;
 }
 
-
 void UVCamComponent::Update()
 {
 	if (!CanUpdate())
@@ -687,12 +613,6 @@ void UVCamComponent::Update()
 		return;
 	}
 
-	// Ensure we register for input if we've not previously registered
-	if (!bIsInputRegistered)
-	{
-		RegisterInputComponent();
-	}
-
 	// If requested then disable the component if we're spawned by sequencer
 	if (bDisableComponentWhenSpawnedBySequencer)
 	{
@@ -705,72 +625,38 @@ void UVCamComponent::Update()
 			return;
 		}
 	}
+	
 	UCineCameraComponent* CameraComponent = GetTargetCamera();
-
 	if (!CameraComponent)
 	{
 		UE_LOG(LogVCamComponent, Error, TEXT("Parent component wasn't valid for Update"));
 		return;
 	}
 
+	EnsureInitialized();
+#if WITH_EDITOR
+	// Somebody may have externally changed subsystem from out under us (user code). Let's make sure the details panel updates.
+	SyncInputSettings();
+#endif
+	
 	const float DeltaTime = GetDeltaTime();
-
-	if (CanEvaluateModifierStack())
+	if (ShouldEvaluateModifierStack())
 	{
 		// Ensure the actor lock reflects the state of the lock property
-		// This is needed as UActorComponent::ConsolidatedPostEditChange will cause the component to be reconstructed on PostEditChange
-		// if the component is inherited
+		// This is needed as UActorComponent::ConsolidatedPostEditChange will cause the component to be reconstructed on PostEditChange if the component is inherited
 		UpdateActorViewportLocks();
 		
 		FLiveLinkCameraBlueprintData InitialLiveLinkData;
 		GetLiveLinkDataForCurrentFrame(InitialLiveLinkData);
-
 		CopyLiveLinkDataToCamera(InitialLiveLinkData, CameraComponent);
 
-		for (FModifierStackEntry& ModifierStackEntry : ModifierStack)
-		{
-			if (UVCamModifier* Modifier = ModifierStackEntry.GeneratedModifier; IsValid(Modifier))
-			{
-				if (ModifierStackEntry.bEnabled)
-				{
-					// Initialize the Modifier if required
-					if (Modifier->IsInitialized())
-					{
-						Modifier->Initialize(ModifierContext, InputComponent);
-						AddInputMappingContext(Modifier);
-					}
-
-					Modifier->Apply(ModifierContext, CameraComponent, DeltaTime);
-				}
-				else
-				{
-					// If the modifier is initialized but not enabled then we deinitialize it
-					if (!Modifier->IsInitialized())
-					{
-						Modifier->Deinitialize();
-					}
-				}
-			}
-		}
-
+		TickModifierStack(DeltaTime);
 		SendCameraDataViaMultiUser();
 	}
 
 	if (ShouldUpdateOutputProviders())
 	{
-		for (UVCamOutputProviderBase* Provider : OutputProviders)
-		{
-			if (Provider)
-			{
-				// Initialize the Provider if required
-				if (!Provider->IsInitialized())
-				{
-					Provider->Initialize();
-				}
-
-				Provider->Tick(DeltaTime);
-			}
-		}
+		TickOutputProviders(DeltaTime);
 	}
 }
 
@@ -780,23 +666,7 @@ void UVCamComponent::SetEnabled(bool bNewEnabled)
 	// NOTE this must be done BEFORE setting the actual bEnabled variable because OutputProviderBase now checks the component enabled state
 	if (!bNewEnabled)
 	{
-		for (UVCamOutputProviderBase* Provider : OutputProviders)
-		{
-			if (IsValid(Provider))
-			{
-				Provider->Deinitialize();
-			}
-		}
-
-		// There's no need to call Initialize if we are being enabled as they'll get automatically intialized
-		// the next time that the stack is evaluated
-		for (FModifierStackEntry& ModifierEntry : ModifierStack)
-		{
-			if (IsValid(ModifierEntry.GeneratedModifier))
-			{
-				ModifierEntry.GeneratedModifier->Deinitialize();
-			}
-		}
+		Deinitialize();
 	}
 
 	bEnabled = bNewEnabled;
@@ -805,19 +675,7 @@ void UVCamComponent::SetEnabled(bool bNewEnabled)
 	// NOTE this must be done AFTER setting the actual bEnabled variable because OutputProviderBase now checks the component enabled state
 	if (bNewEnabled)
 	{
-		if (ShouldUpdateOutputProviders() && CanUpdate())
-		{
-			for (UVCamOutputProviderBase* Provider : OutputProviders)
-			{
-				if (IsValid(Provider))
-				{
-					Provider->Initialize();
-				}
-			}
-		}
-
-		// Register the input component now we're enabled
-		RegisterInputComponent();
+		Initialize();
 	}
 }
 
@@ -1128,6 +986,37 @@ void UVCamComponent::SetInputProfile(const FVCamInputProfile& NewInputProfile)
 	ApplyInputProfile();
 }
 
+const FVCamInputDeviceConfig& UVCamComponent::GetInputDeviceSettings() const
+{
+	// This input system exists while this component is enabled.
+	// User code may call SetInputSettings directly on the subsystem - we want to capture that case. 
+	if (UInputVCamSubsystem* InputSubsystem = GetInputVCamSubsystem())
+	{
+		return InputSubsystem->GetInputSettings();
+	}
+	
+	return InputDeviceSettings;
+}
+
+void UVCamComponent::SetInputDeviceSettings(const FVCamInputDeviceConfig& NewInputProfile)
+{
+	InputDeviceSettings = NewInputProfile;
+	if (UInputVCamSubsystem* EnhancedInputSubsystemInterface = GetInputVCamSubsystem())
+	{
+		EnhancedInputSubsystemInterface->SetInputSettings(InputDeviceSettings);
+	}
+}
+
+TArray<UVCamSubsystem*> UVCamComponent::GetSubsystemArray(const TSubclassOf<UVCamSubsystem>& Class) const
+{
+	return SubsystemCollection.GetSubsystemArray<UVCamSubsystem>(Class);
+}
+
+UInputVCamSubsystem* UVCamComponent::GetInputVCamSubsystem() const
+{
+	return SubsystemCollection.GetSubsystem<UInputVCamSubsystem>(UInputVCamSubsystem::StaticClass());
+}
+
 void UVCamComponent::GetLiveLinkDataForCurrentFrame(FLiveLinkCameraBlueprintData& LiveLinkData)
 {
 	IModularFeatures& ModularFeatures = IModularFeatures::Get();
@@ -1182,7 +1071,7 @@ void UVCamComponent::UnregisterObjectForInput(UObject* Object) const
 
 TArray<FEnhancedActionKeyMapping> UVCamComponent::GetPlayerMappableKeys() const
 {
-	if (const IEnhancedInputSubsystemInterface* EnhancedInputSubsystemInterface = GetEnhancedInputSubsystemInterface())
+	if (const IEnhancedInputSubsystemInterface* EnhancedInputSubsystemInterface = GetInputVCamSubsystem())
 	{
 		return EnhancedInputSubsystemInterface->GetAllPlayerMappableActionKeyMappings();
 	}
@@ -1191,7 +1080,7 @@ TArray<FEnhancedActionKeyMapping> UVCamComponent::GetPlayerMappableKeys() const
 
 void UVCamComponent::InjectInputForAction(const UInputAction* Action, FInputActionValue RawValue, const TArray<UInputModifier*>& Modifiers, const TArray<UInputTrigger*>& Triggers)
 {
-	if (IEnhancedInputSubsystemInterface* EnhancedInputSubsystemInterface = GetEnhancedInputSubsystemInterface())
+	if (IEnhancedInputSubsystemInterface* EnhancedInputSubsystemInterface = GetInputVCamSubsystem())
 	{
 		EnhancedInputSubsystemInterface->InjectInputForAction( Action, RawValue, Modifiers, Triggers);
 	}
@@ -1199,7 +1088,7 @@ void UVCamComponent::InjectInputForAction(const UInputAction* Action, FInputActi
 
 void UVCamComponent::InjectInputVectorForAction(const UInputAction* Action, FVector Value, const TArray<UInputModifier*>& Modifiers, const TArray<UInputTrigger*>& Triggers)
 {
-	if (IEnhancedInputSubsystemInterface* EnhancedInputSubsystemInterface = GetEnhancedInputSubsystemInterface())
+	if (IEnhancedInputSubsystemInterface* EnhancedInputSubsystemInterface = GetInputVCamSubsystem())
 	{
 		EnhancedInputSubsystemInterface->InjectInputVectorForAction(Action, Value, Modifiers, Triggers);
 	}
@@ -1207,7 +1096,7 @@ void UVCamComponent::InjectInputVectorForAction(const UInputAction* Action, FVec
 
 void UVCamComponent::ApplyInputProfile()
 {
-	IEnhancedInputSubsystemInterface* EnhancedInputSubsystemInterface = GetEnhancedInputSubsystemInterface();
+	IEnhancedInputSubsystemInterface* EnhancedInputSubsystemInterface = GetInputVCamSubsystem();
 	if (EnhancedInputSubsystemInterface)
 	{
 		EnhancedInputSubsystemInterface->RemoveAllPlayerMappedKeys();
@@ -1225,6 +1114,129 @@ void UVCamComponent::ApplyInputProfile()
 	}
 }
 
+void UVCamComponent::EnsureInitialized()
+{
+	if (!SubsystemCollection.IsInitialized())
+	{
+		Initialize();
+	}
+}
+
+void UVCamComponent::Initialize()
+{
+	SubsystemCollection.Initialize(this);
+	// 1. Input is required
+	RegisterInputComponent();
+
+	// 2. Output provider overlay widgets will access the modifiers, so let's init them first
+	const bool bInitModifiers = ShouldEvaluateModifierStack() && CanUpdate(); 
+	if (bInitModifiers)
+	{
+		for (FModifierStackEntry& ModifierStackEntry : ModifierStack)
+		{
+			if (UVCamModifier* Modifier = ModifierStackEntry.GeneratedModifier
+				; IsValid(Modifier) && ModifierStackEntry.bEnabled && !Modifier->IsInitialized())
+			{
+				Modifier->Initialize(ModifierContext, InputComponent);
+			}
+		}
+	}
+
+	// 3. Output providers
+	const bool bInitOutputProviders = bInitModifiers && ShouldUpdateOutputProviders();
+	if (bInitOutputProviders)
+	{
+		for (UVCamOutputProviderBase* Provider : OutputProviders)
+		{
+			if (IsValid(Provider))
+			{
+				Provider->Initialize();
+			}
+		}
+	}
+}
+
+void UVCamComponent::Deinitialize()
+{
+	for (UVCamOutputProviderBase* Provider : OutputProviders)
+	{
+		if (IsValid(Provider))
+		{
+			Provider->Deinitialize();
+		}
+	}
+
+	for (FModifierStackEntry& ModifierEntry : ModifierStack)
+	{
+		if (IsValid(ModifierEntry.GeneratedModifier))
+		{
+			ModifierEntry.GeneratedModifier->Deinitialize();
+		}
+	}
+
+	// Viewports may be manipulated by output providers or modifiers
+	UnlockAllViewports();
+	
+	UnregisterInputComponent();
+	SubsystemCollection.Deinitialize();
+}
+
+void UVCamComponent::SyncInputSettings()
+{
+	if (UInputVCamSubsystem* InputSubsystem = GetInputVCamSubsystem())
+	{
+		SetInputDeviceSettings(InputDeviceSettings);
+	}
+}
+
+void UVCamComponent::TickModifierStack(const float DeltaTime)
+{
+	UCineCameraComponent* CameraComponent = GetTargetCamera();
+	check(CameraComponent);
+	
+	for (FModifierStackEntry& ModifierStackEntry : ModifierStack)
+	{
+		if (UVCamModifier* Modifier = ModifierStackEntry.GeneratedModifier; IsValid(Modifier))
+		{
+			if (ModifierStackEntry.bEnabled)
+			{
+				if (!Modifier->IsInitialized())
+				{
+					Modifier->Initialize(ModifierContext, InputComponent);
+					AddInputMappingContext(Modifier);
+				}
+
+				Modifier->Apply(ModifierContext, CameraComponent, DeltaTime);
+			}
+			else
+			{
+				// If the modifier is initialized but not enabled then we deinitialize it
+				if (!Modifier->IsInitialized())
+				{
+					Modifier->Deinitialize();
+				}
+			}
+		}
+	}
+}
+
+void UVCamComponent::TickOutputProviders(float DeltaTime)
+{
+	for (UVCamOutputProviderBase* Provider : OutputProviders)
+	{
+		if (Provider)
+		{
+			// Initialize the Provider if required
+			if (!Provider->IsInitialized())
+			{
+				Provider->Initialize();
+			}
+
+			Provider->Tick(DeltaTime);
+		}
+	}
+}
+
 void UVCamComponent::CopyLiveLinkDataToCamera(const FLiveLinkCameraBlueprintData& LiveLinkData, UCineCameraComponent* CameraComponent)
 {
 	const FLiveLinkCameraStaticData& StaticData = LiveLinkData.StaticData;
@@ -1290,26 +1302,6 @@ void UVCamComponent::DestroyOutputProvider(UVCamOutputProviderBase* Provider)
 	}
 }
 
-void UVCamComponent::ResetAllOutputProviders()
-{
-	for (UVCamOutputProviderBase* Provider : OutputProviders)
-	{
-		if (Provider)
-		{
-			// Initialization will also recover active state 
-			Provider->Deinitialize();
-
-			// We only Initialize a provider if they're able to be updated
-			// If they later become able to be updated then they will be
-			// Initialized inside the Update() loop
-			if (ShouldUpdateOutputProviders())
-			{
-				Provider->Initialize();
-			}
-		}
-	}
-}
-
 void UVCamComponent::ValidateModifierStack(const FString BaseName /*= "NewModifier"*/)
 {
 	int32 ModifiedStackIndex;
@@ -1439,23 +1431,13 @@ void UVCamComponent::OnMapChanged(UWorld* World, EMapChangeType ChangeType)
 void UVCamComponent::OnBeginPIE(const bool bInIsSimulating)
 {
 	UWorld* World = GetWorld();
-
-	if (!World)
-	{
-		return;
-	}
-
-	if (World->WorldType == EWorldType::Editor)
+	if (World && World->WorldType == EWorldType::Editor)
 	{
-		// Deinitialize all output providers in the editor world
-		for (UVCamOutputProviderBase* Provider : OutputProviders)
-		{
-			if (Provider)
-			{
-				Provider->Deinitialize();
-			}
-		}
-
+		// Reasons:
+		// 1. Output providers, like pixel streaming, may interfere
+		// 2. Input will interfere
+		Deinitialize();
+		
 		// Ensure the Editor components do not update during PIE
 		bIsEditorObjectButPIEIsRunning = true;
 	}
@@ -1472,19 +1454,13 @@ void UVCamComponent::OnEndPIE(const bool bInIsSimulating)
 
 	if (World->WorldType == EWorldType::PIE)
 	{
-		// Disable all output providers in the PIE world
-		for (UVCamOutputProviderBase* Provider : OutputProviders)
-		{
-			if (Provider)
-			{
-				Provider->Deinitialize();
-			}
-		}
+		// If Update is called (should not happen), disabling instead of Deinitialize to ensure Update() will be skipped
+		SetEnabled(false);
 	}
 	else if (World->WorldType == EWorldType::Editor)
 	{
-		// Allow the Editor components to start updating again
 		bIsEditorObjectButPIEIsRunning = false;
+		Initialize();
 	}
 }
 
@@ -1610,32 +1586,18 @@ void UVCamComponent::SendCameraDataViaMultiUser()
 #endif
 }
 
-bool UVCamComponent::IsCameraInVPRole() const
-{
-#if WITH_EDITOR
-	UVPSettings* Settings = UVPSettings::GetVPSettings();
-	// We are in a valid camera role if the user has not assigned a role or the current VPSettings role matches the
-	// assigned role.
-	//
-	UVirtualProductionRolesSubsystem* VPRolesSubsystem = GEngine->GetEngineSubsystem<UVirtualProductionRolesSubsystem>();
-	return !Role.IsValid() || (VPRolesSubsystem && VPRolesSubsystem->HasActiveRole(Role));
-#else
-	return true;
-#endif
-}
-
-bool UVCamComponent::CanEvaluateModifierStack() const
-{
-	return !IsMultiUserSession() || (IsMultiUserSession() && IsCameraInVPRole());
-}
-
 bool UVCamComponent::ShouldUpdateOutputProviders() const
 {
 	// We should only update output providers in 3 situations
 	// - We're not in Multi User
 	// - We have the virtual camera role
 	// - We have explicitly set that we want to run Output Providers even when not in the camera role
-	return !IsMultiUserSession() || IsCameraInVPRole() || !bDisableOutputOnMultiUserReceiver;
+	return ShouldEvaluateModifierStack() || !bDisableOutputOnMultiUserReceiver;
+}
+
+bool UVCamComponent::ShouldEvaluateModifierStack() const
+{
+	return !IsMultiUserSession() || IsCameraInVPRole();
 }
 
 bool UVCamComponent::IsMultiUserSession() const
@@ -1647,3 +1609,15 @@ bool UVCamComponent::IsMultiUserSession() const
 #endif
 }
 
+bool UVCamComponent::IsCameraInVPRole() const
+{
+#if WITH_EDITOR
+	// We are in a valid camera role if the user has not assigned a role or the current VPSettings role matches the
+	// assigned role.
+	UVirtualProductionRolesSubsystem* VPRolesSubsystem = GEngine->GetEngineSubsystem<UVirtualProductionRolesSubsystem>();
+	return !Role.IsValid() || (VPRolesSubsystem && VPRolesSubsystem->HasActiveRole(Role));
+#else
+	return true;
+#endif
+}
+
diff --git a/Engine/Plugins/Experimental/VirtualCamera/Source/VCamCore/Private/VCamSubsystem.cpp b/Engine/Plugins/Experimental/VirtualCamera/Source/VCamCore/Private/VCamSubsystem.cpp
new file mode 100644
--- /dev/null
+++ b/Engine/Plugins/Experimental/VirtualCamera/Source/VCamCore/Private/VCamSubsystem.cpp
@@ -0,0 +1,10 @@
+// Copyright Epic Games, Inc. All Rights Reserved.
+
+#include "VCamSubsystem.h"
+
+#include "VCamComponent.h"
+
+UVCamComponent* UVCamSubsystem::GetVCamComponent() const
+{
+	return GetTypedOuter<UVCamComponent>();
+}
diff --git a/Engine/Plugins/Experimental/VirtualCamera/Source/VCamCore/Public/Input/InputVCamSubsystem.h b/Engine/Plugins/Experimental/VirtualCamera/Source/VCamCore/Public/Input/InputVCamSubsystem.h
new file mode 100644
--- /dev/null
+++ b/Engine/Plugins/Experimental/VirtualCamera/Source/VCamCore/Public/Input/InputVCamSubsystem.h
@@ -0,0 +1,73 @@
+// Copyright Epic Games, Inc. All Rights Reserved.
+
+#pragma once
+#include "Input/VCamInputProcessor.h"
+
+struct FVCamInputDeviceConfig;
+class UVCamPlayerInput;
+
+#include "VCamSubsystem.h"
+#include "EnhancedInputSubsystemInterface.h"
+#include "Tickable.h"
+#include "InputVCamSubsystem.generated.h"
+
+namespace UE::VCamCore::Private
+{
+	class FVCamInputProcessor;
+}
+
+class UVCamPlayerInput;
+
+/**
+ * Handles all input for UVCamComponent.
+ * Essentially maps input devices to UVCamComponents, similar like APlayerController does for gameplay code.
+ */
+UCLASS()
+class VCAMCORE_API UInputVCamSubsystem : public UVCamSubsystem, public IEnhancedInputSubsystemInterface, public FTickableGameObject
+{
+	GENERATED_BODY()
+public:
+	
+	//~ Begin USubsystem Interface
+	virtual void Initialize(FSubsystemCollectionBase& Collection) override;
+	virtual void Deinitialize() override;
+	//~ End USubsystem Interface
+	
+	/** Inputs a key on this subsystem's player input which can then be processed as normal during Tick. */
+	bool InputKey(const FInputKeyParams& Params);
+
+	/** Pushes this input component onto the stack to be processed by this subsystem's tick function */
+	void PushInputComponent(UInputComponent* InInputComponent);
+	/** Removes this input component onto the stack to be processed by this subsystem's tick function */
+	bool PopInputComponent(UInputComponent* InInputComponent);
+	
+	UFUNCTION(BlueprintCallable, Category = "Input")
+	void SetShouldConsumeGamepadInput(EVCamGamepadInputMode GamepadInputMode);
+
+	const FVCamInputDeviceConfig& GetInputSettings() const;
+	void SetInputSettings(const FVCamInputDeviceConfig& Input);
+	
+	//~ Begin FTickableGameObject interface
+	virtual UWorld* GetTickableGameObjectWorld() const override;
+	virtual bool IsTickableInEditor() const { return true; }
+	virtual ETickableTickType GetTickableTickType() const override;
+	virtual bool IsAllowedToTick() const override;
+	virtual void Tick(float DeltaTime) override;
+	TStatId GetStatId() const override { RETURN_QUICK_DECLARE_CYCLE_STAT(UInputVCamSubsystem, STATGROUP_Tickables); }
+	//~ End FTickableGameObject interface
+
+	//~ Begin IEnhancedInputSubsystemInterface Interface
+	virtual UEnhancedPlayerInput* GetPlayerInput() const override;
+	//~ End IEnhancedInputSubsystemInterface Interface
+
+private:
+
+	UPROPERTY(Transient, Instanced)
+	TObjectPtr<UVCamPlayerInput> PlayerInput;
+	
+	TSharedPtr<UE::VCamCore::Private::FVCamInputProcessor> InputPreprocessor;
+	
+	/** Internal. This is the current stack of InputComponents that is being processed by the PlayerInput. */
+	UPROPERTY(Transient)
+	TArray<TWeakObjectPtr<UInputComponent>> CurrentInputStack;
+};
diff --git a/Engine/Plugins/Experimental/VirtualCamera/Source/VCamCore/Public/Input/VCamInputDeviceConfig.h b/Engine/Plugins/Experimental/VirtualCamera/Source/VCamCore/Public/Input/VCamInputDeviceConfig.h
new file mode 100644
--- /dev/null
+++ b/Engine/Plugins/Experimental/VirtualCamera/Source/VCamCore/Public/Input/VCamInputDeviceConfig.h
@@ -0,0 +1,96 @@
+// Copyright Epic Games, Inc. All Rights Reserved.
+
+#pragma once
+
+#include "VCamInputDeviceConfig.generated.h"
+
+UENUM()
+enum class EVCamInputLoggingMode : uint8
+{
+	/** No logging */
+	None,
+	
+	/** Log only input that are passed down to the input actions (i.e. that passed filtering conditions) */
+	OnlyConsumable,
+
+	/** Log all gamepad input regardless whether it is passed down to input actions or not (i.e. that passed filtering conditions) */
+	OnlyGamepad,
+	
+	/** Log all input, regardless whether it is passed down to input actions or not (i.e. that passed filtering conditions). */
+	All,
+	
+	/** Like All but never reports mouse data (generated every tick - it could spam the output log). */
+	AllExceptMouse
+};
+
+UENUM()
+enum class EVCamInputMode : uint8
+{
+	/** Input is forwarded to input action. */
+	Allow,
+	/** Input is not forwarded to input actions. */
+	Ignore
+};
+
+UENUM()
+enum class EVCamGamepadInputMode : uint8
+{
+	/** Input is forwarded to input action. */
+	Allow,
+	/** Input is not forwarded to input actions. */
+	Ignore,
+	/** Input is not forwarded to input actions and we block anybody else from receiving this input. */
+	IgnoreAndConsume
+};
+
+/** Defines the input devices a UVCamComponent will accept input from. */
+USTRUCT(BlueprintType)
+struct VCAMCORE_API FVCamInputDeviceConfig
+{
+	GENERATED_BODY()
+	
+	/**
+	 * List of input devices from which input can trigger input actions.
+	 * Typically this is used for gamepads.
+	 *
+	 * Input device IDs start at 0 and increase by 1 as more devices connect. When a device disconnect, the ID is recycled
+	 * and becomes available for reassignment to the next device that connects; when a device connects, the lowest possible ID is reassigned.
+	 *
+	 * Example: suppose you have three gamepads called A, B, and VCamDevicePairingConfig.h
+	 * 1. Connect gamepad A > receives ID 0
+	 * 2. Connect gamepad B > receives ID 1
+	 * 3. Disconnect gamepad A > gamepad B will still have ID 1
+	 * 4. Connect the same gamepad A OR another gamepad C > receives ID 0.
+	 *
+	 * Note: Keyboards always have ID = 0, mice ID = -1.
+	 * Note: The first gamepad will have ID = 0 even though keyboards will also have ID 0.
+	 */
+	UPROPERTY(EditAnywhere, Category = "Input")
+	TArray<int32> AllowedInputDeviceIds;
+	
+	/**
+	 * Should keyboard input trigger input actions?
+	 * Keyboards are always mapped to input device ID 0.
+	 */
+	UPROPERTY(EditAnywhere, Category = "Input")
+	EVCamInputMode KeyboardInputMode = EVCamInputMode::Allow;
+
+	/**
+	 * Should mouse input trigger input actions?
+	 * Mice are always mapped to input device ID -1.
+	 */
+	UPROPERTY(EditAnywhere, Category = "Input")
+	EVCamInputMode MouseInputMode = EVCamInputMode::Allow;
+
+	/** What should be done with gamepad input. */
+	UPROPERTY(EditAnywhere, Category = "Input", AdvancedDisplay)
+	EVCamGamepadInputMode GamepadInputMode = EVCamGamepadInputMode::Allow;
+
+	/**
+	 * What type of input should be logged.
+	 *
+	 * Tip: Filter the log by LogVCamInputDebug.
+	 */
+	UPROPERTY(EditAnywhere, Category = "Input", AdvancedDisplay)
+	EVCamInputLoggingMode LoggingMode = EVCamInputLoggingMode::OnlyGamepad;
+};
\ No newline at end of file
diff --git a/Engine/Plugins/Experimental/VirtualCamera/Source/VCamCore/Public/Input/VCamPlayerInput.h b/Engine/Plugins/Experimental/VirtualCamera/Source/VCamCore/Public/Input/VCamPlayerInput.h
new file mode 100644
--- /dev/null
+++ b/Engine/Plugins/Experimental/VirtualCamera/Source/VCamCore/Public/Input/VCamPlayerInput.h
@@ -0,0 +1,32 @@
+// Copyright Epic Games, Inc. All Rights Reserved.
+
+#pragma once
+
+#include "EnhancedPlayerInput.h"
+#include "VCamInputDeviceConfig.h"
+#include "VCamPlayerInput.generated.h"
+
+/**
+ * Receives raw input from input processor (or player controller in user code).
+ * Filters input
+ */
+UCLASS()
+class VCAMCORE_API UVCamPlayerInput : public UEnhancedPlayerInput
+{
+	GENERATED_BODY()
+public:
+
+	//~ Begin UEnhancedPlayerInput Interface
+	virtual bool InputKey(const FInputKeyParams& Params) override;
+	virtual void ProcessInputStack(const TArray<UInputComponent*>& InputComponentStack, const float DeltaTime, const bool bGamePaused) override;
+	//~ End UEnhancedPlayerInput Interface
+
+	const FVCamInputDeviceConfig& GetInputSettings() const { return InputDeviceSettings; }
+	void SetInputSettings(const FVCamInputDeviceConfig& Input);
+
+private:
+
+	/** The device settings governing what we do with received input, e.g. filtering depending on input device, logging, etc. */
+	UPROPERTY(Transient)
+	FVCamInputDeviceConfig InputDeviceSettings;
+};
diff --git a/Engine/Plugins/Experimental/VirtualCamera/Source/VCamCore/Public/Modifier/VCamModifier.h b/Engine/Plugins/Experimental/VirtualCamera/Source/VCamCore/Public/Modifier/VCamModifier.h
--- a/Engine/Plugins/Experimental/VirtualCamera/Source/VCamCore/Public/Modifier/VCamModifier.h
+++ b/Engine/Plugins/Experimental/VirtualCamera/Source/VCamCore/Public/Modifier/VCamModifier.h
@@ -39,9 +39,10 @@ public:
 	//~ Begin UObject Interface
 	virtual void BeginDestroy() override;
 	virtual void PostLoad() override;
+	virtual UWorld* GetWorld() const override;
 	//~ End UObject Interface
 
-	bool IsInitialized() const { return bRequiresInitialization; };
+	bool IsInitialized() const { return !bIsInitialized; };
 
 	UFUNCTION(BlueprintCallable, Category="VirtualCamera")
 	UVCamComponent* GetOwningVCamComponent() const;
@@ -51,7 +52,6 @@ public:
 
 	UFUNCTION(BlueprintCallable, Category = "VirtualCamera")
 	void SetEnabled(bool bNewEnabled);
-
 	UFUNCTION(BlueprintPure, Category = "VirtualCamera", meta=(ReturnDisplayName="Enabled"))
 	bool IsEnabled() const;
 
@@ -59,7 +59,6 @@ public:
 	// Returns a bool for whether 
 	UFUNCTION(BlueprintCallable, Category = "VirtualCamera", meta=(ReturnDisplayName="Success"))
 	bool SetStackEntryName(FName NewName);
-
 	// Gets the name of the modifier in the associated modifier stack
 	UFUNCTION(BlueprintPure, Category = "VirtualCamera", meta=(ReturnDisplayName="Name"))
 	FName GetStackEntryName() const;
@@ -75,12 +74,12 @@ public:
 	UPROPERTY(EditAnywhere, BlueprintReadWrite, Category = "VCam Connection Points")
 	TMap<FName, FVCamModifierConnectionPoint> ConnectionPoints;
 
-	virtual UWorld* GetWorld() const override;
-	
 private:
+	
 	FModifierStackEntry* GetCorrespondingStackEntry() const;
 
-	bool bRequiresInitialization = true;
+	UPROPERTY(Transient)
+	bool bIsInitialized = true;
 };
 
 UCLASS(EditInlineNew)
diff --git a/Engine/Plugins/Experimental/VirtualCamera/Source/VCamCore/Public/VCamComponent.h b/Engine/Plugins/Experimental/VirtualCamera/Source/VCamCore/Public/VCamComponent.h
--- a/Engine/Plugins/Experimental/VirtualCamera/Source/VCamCore/Public/VCamComponent.h
+++ b/Engine/Plugins/Experimental/VirtualCamera/Source/VCamCore/Public/VCamComponent.h
@@ -12,6 +12,8 @@
 #include "GameplayTagContainer.h"
 #include "Roles/LiveLinkCameraTypes.h"
 #include "VCamInputSettings.h"
+#include "VCamSubsystem.h"
+#include "Input/VCamInputDeviceConfig.h"
 #include "Util/VCamViewportLocker.h"
 
 #if WITH_EDITOR
@@ -22,6 +24,7 @@
 
 #include "VCamComponent.generated.h"
 
+class UInputVCamSubsystem;
 class UVCamOutputProviderBase;
 DECLARE_LOG_CATEGORY_EXTERN(LogVCamComponent, Log, All);
 
@@ -37,6 +40,25 @@ class FLevelEditorViewportClient;
 
 DECLARE_DYNAMIC_MULTICAST_DELEGATE_OneParam(FOnComponentReplaced, UVCamComponent*, NewComponent);
 
+/**
+ * Provides a modular system for editing a UCineCameraComponent using user widgets.
+ * This component must be attached as a direct child of UCineCameraComponent.
+ *
+ * This component implements a Model-View-Controller architecture where modifiers are the model, output providers the view, and
+ * this component the controller.
+ *
+ * There are three aspects to this component:
+ * - Modifiers implement logic for changing properties on the UCineCameraComponent. Modifiers contain ConnectionPoints.
+ *   ConnectionPoints can optionally expose UInputActions (Enhanced Input) that can be invoked by widgets.
+ * - Output providers create and render widgets (possibly streaming them). Usually output providers create UVCamWidgets, which
+ *   are special widgets that can connect to ConnectionPoints. Widgets interact with modifiers in two ways:
+ *     - Simple: trigger input actions that modifiers are subscribed to and expose via connections.
+ *     - Advanced: query whether modifiers implement certain custom defined interfaces.
+ *       UVCamWidget Connections can be configured with required and optional interfaces.
+ * - UVCamSubsystems exist for as long as a UVCamComponent is enabled (this is comparable to ULocalPlayerSubsystem).
+ *   One notable such system is the UInputVCamSubsystem which allows UVCamComponents to bind to input devices similarly to how
+ *   APlayerControllers do in shipped games.
+ */
 UCLASS(Blueprintable, ClassGroup = VCam, HideCategories=(Mobility), meta=(BlueprintSpawnableComponent))
 class VCAMCORE_API UVCamComponent : public USceneComponent
 {
@@ -127,7 +149,6 @@ public:
 	FKey GetPlayerMappedKey(const FName MappingName) const;
 
 
-	
 	/** Sets whether the VCamComponent will update every frame */
 	UFUNCTION(BlueprintSetter)
 	void SetEnabled(bool bNewEnabled);
@@ -316,7 +337,31 @@ public:
 	const FVCamInputProfile& GetInputProfile() const { return InputProfile; }
 	UFUNCTION(BlueprintCallable, BlueprintInternalUseOnly, Category="VirtualCamera")
 	void SetInputProfile(const FVCamInputProfile& NewInputProfile);
+
+	/** Gets the current input device settings being used (if this component is enabled) or that will be used (if not enabled). */
+	UFUNCTION(BlueprintPure, Category = "VCam Input")
+	const FVCamInputDeviceConfig& GetInputDeviceSettings() const;
+	/** Propagates devices settings all the way to the player input. Causes input to be filtered / consumed differently. */
+	UFUNCTION(BlueprintCallable, BlueprintInternalUseOnly, Category="VirtualCamera")
+	void SetInputDeviceSettings(const FVCamInputDeviceConfig& NewInputProfile);
+
 	
+	/******************** Subsystem Queries ********************/
+
+	/** Gets all subsystems implementing this interface */
+	UFUNCTION(BlueprintPure, Category = "VirtualCamera")
+	TArray<UVCamSubsystem*> GetSubsystemArray(const TSubclassOf<UVCamSubsystem>& Class) const;
+
+	/** Gets the subsystem responsible for input handling. */
+	UFUNCTION(BlueprintPure, Category = "VirtualCamera")
+	UInputVCamSubsystem* GetInputVCamSubsystem() const;
+	
+	template <typename TSubsystemClass>
+	TSubsystemClass* GetSubsystem(const TSubclassOf<TSubsystemClass>& SubsystemClass) const { return SubsystemCollection.GetSubsystem(SubsystemClass); }
+	
+	template <typename TSubsystemClass>
+	const TArray<TSubsystemClass*>& GetSubsystemArray(const TSubclassOf<TSubsystemClass>& SubsystemClass) const { return SubsystemCollection.GetSubsystemArray(SubsystemClass); }
+
 private:
 	
 	/** Whether the VCamComponent will update every frame */
@@ -379,6 +424,9 @@ private:
 	UPROPERTY(EditAnywhere, BlueprintGetter = "GetInputProfile", BlueprintSetter = "SetInputProfile", Category = "VirtualCamera")
 	FVCamInputProfile InputProfile;
 
+	UPROPERTY(EditAnywhere, BlueprintGetter = "GetInputDeviceSettings", BlueprintSetter = "SetInputDeviceSettings", Category = "VirtualCamera")
+	FVCamInputDeviceConfig InputDeviceSettings;
+
 	/** List of Output Providers (executed in order) */
 	UPROPERTY(EditAnywhere, Instanced, Category = "VirtualCamera")
 	TArray<TObjectPtr<UVCamOutputProviderBase>> OutputProviders;
@@ -410,7 +458,17 @@ private:
 	/** Variable used for pausing update on editor objects while PIE is running */
 	bool bIsEditorObjectButPIEIsRunning = false;
 
-	bool bIsInputRegistered = false;
+	/** Initialize and deinitialize calls match our  */
+	FSubsystemCollection<UVCamSubsystem> SubsystemCollection;
+
+	void EnsureInitialized();
+	virtual void Initialize();
+	virtual void Deinitialize();
+
+	void SyncInputSettings();
+	
+	void TickModifierStack(float DeltaTime);
+	void TickOutputProviders(float DeltaTime);
 	
 	static void CopyLiveLinkDataToCamera(const FLiveLinkCameraBlueprintData& LiveLinkData, UCineCameraComponent* CameraComponent);
 
@@ -418,7 +476,6 @@ private:
 	void UpdateActorViewportLocks();
 	void UnlockAllViewports();
 	void DestroyOutputProvider(UVCamOutputProviderBase* Provider);
-	void ResetAllOutputProviders();
 
 	// Use the Saved Modifier Stack from PreEditChange to find the modified entry and then ensure the modified entry's name is unique
 	// If a new modifier has been created then its name will be defaulted to BaseName
@@ -456,31 +513,21 @@ private:
 	double PreviousUpdateTime = 0;
 #endif
 
-	/** Is the camera currently in a role assigned to the session. */
-	bool IsCameraInVPRole() const;
 
 	/** Send the current camera state via Multi-user if connected and in a */
 	void SendCameraDataViaMultiUser();
-
-	/** Are we in a multi-user session. */
-	bool IsMultiUserSession() const;
-
-	/** Can the modifier stack be evaluated. */
-	bool CanEvaluateModifierStack() const;
-
-	/** Output Providers should only update outside of MU or if we have the correct VP Role */
+	
+	/** Can return false if we're in a multi-user session and the user configured us not to update. */
 	bool ShouldUpdateOutputProviders() const;
+	/** Can return false if we're in a multi-user session and the user configured us not to update. */
+	bool ShouldEvaluateModifierStack() const;
+	bool IsMultiUserSession() const;
+	bool IsCameraInVPRole() const;
 
 	// When another component replaces us, get a notification so we can clean up
 	void NotifyComponentWasReplaced(UVCamComponent* ReplacementComponent);
 
-	/*
-	 * Gets the input interface for the currently active input method.
-	 * Will switch between editor and player based input as needed
-	 */
-	IEnhancedInputSubsystemInterface* GetEnhancedInputSubsystemInterface() const;
-
 	/** Utility functions for registering and unregistering our input component with the correct input system */
-	void RegisterInputComponent();
-	void UnregisterInputComponent();
+	virtual void RegisterInputComponent();
+	virtual void UnregisterInputComponent();
 };
diff --git a/Engine/Plugins/Experimental/VirtualCamera/Source/VCamCore/Public/VCamSubsystem.h b/Engine/Plugins/Experimental/VirtualCamera/Source/VCamCore/Public/VCamSubsystem.h
new file mode 100644
--- /dev/null
+++ b/Engine/Plugins/Experimental/VirtualCamera/Source/VCamCore/Public/VCamSubsystem.h
@@ -0,0 +1,19 @@
+// Copyright Epic Games, Inc. All Rights Reserved.
+
+#pragma once
+
+#include "Subsystems/Subsystem.h"
+#include "VCamSubsystem.generated.h"
+
+class UVCamComponent;
+
+/** Base class for auto instanced and initialized subsystem that share the lifetime of UVCamComponents. */
+UCLASS(Abstract, Blueprintable, Within = VCamComponent)
+class VCAMCORE_API UVCamSubsystem : public USubsystem
+{
+	GENERATED_BODY()
+public:
+
+	UFUNCTION(BlueprintPure, Category = "VirtualCamera")
+	UVCamComponent* GetVCamComponent() const;
+};
