Alternative style settings with which you can target modifiers directly

diff --git a/Engine/Plugins/Experimental/VirtualCamera/Source/VCamExtensions/Private/Styling/ClassBasedWidgetStyleDefinitions.cpp b/Engine/Plugins/Experimental/VirtualCamera/Source/VCamExtensions/Private/Styling/ClassBasedWidgetStyleDefinitions.cpp
--- a/Engine/Plugins/Experimental/VirtualCamera/Source/VCamExtensions/Private/Styling/ClassBasedWidgetStyleDefinitions.cpp
+++ b/Engine/Plugins/Experimental/VirtualCamera/Source/VCamExtensions/Private/Styling/ClassBasedWidgetStyleDefinitions.cpp
@@ -5,26 +5,16 @@
 #include "VCamComponent.h"
 #include "Algo/ForEach.h"
 
-DEFINE_LOG_CATEGORY_STATIC(LogClassBasedModifierMetaDataRules, All, All);
-
 namespace UE::VCamExtensions::Private
 {
 	static TArray<UWidgetStyleData*> AccumulateMetaDataRecursively(
 		const TMap<TSubclassOf<UVCamModifier>, FPerModifierClassWidgetSytleData>& Config,
-		UVCamComponent* VCamComponent, 
-		FName ModifierId,
+		UVCamModifier* Modifier,
 		TFunctionRef<const FWidgetStyleDataConfig*(const FPerModifierClassWidgetSytleData& ClassConfig)> RetrieveConfig)
 	{
 		TArray<UWidgetStyleData*> Result;
-		if (!IsValid(VCamComponent))
-		{
-			return Result;
-		}
-
-		UVCamModifier* Modifier = VCamComponent->GetModifierByName(ModifierId);
 		if (!Modifier)
 		{
-			UE_LOG(LogClassBasedModifierMetaDataRules, Warning, TEXT("Unknown modifier %s (on component %s)"), *ModifierId.ToString(), *VCamComponent->GetPathName());
 			return Result;
 		}
 
@@ -53,31 +43,29 @@ namespace UE::VCamExtensions::Private
 }
 
 
-TArray<UWidgetStyleData*> UClassBasedWidgetStyleDefinitions::GetStylesForModifier_Implementation(UVCamComponent* VCamComponent, FName ModifierId) const
+TArray<UWidgetStyleData*> UClassBasedWidgetStyleDefinitions::GetStylesForModifier_Implementation(UVCamModifier* Modifier) const
 {
 	return UE::VCamExtensions::Private::AccumulateMetaDataRecursively(
 		Config,
-		VCamComponent,
-		ModifierId,
+		Modifier,
 		[](const FPerModifierClassWidgetSytleData& ClassConfig)
 		{
 			return &ClassConfig.ModifierStyles;
 		});
 }
 
-TArray<UWidgetStyleData*> UClassBasedWidgetStyleDefinitions::GetStylesForCategoryInModifier_Implementation(UVCamComponent* VCamComponent, FName ModifierId, FName ConnectionPointId) const
+TArray<UWidgetStyleData*> UClassBasedWidgetStyleDefinitions::GetStylesForConnectionPoint_Implementation(UVCamModifier* Modifier, FName ConnectionPointId) const
 {
 	return UE::VCamExtensions::Private::AccumulateMetaDataRecursively(
 		Config,
-		VCamComponent,
-		ModifierId,
+		Modifier,
 		[ConnectionPointId](const FPerModifierClassWidgetSytleData& ClassConfig)
 		{
-			return ClassConfig.CategorizedStyles.Find(ConnectionPointId);
+			return ClassConfig.ConnectionPointStyles.Find(ConnectionPointId);
 		});
 }
 
-TArray<UWidgetStyleData*> UClassBasedWidgetStyleDefinitions::GetStylesForCategoryWithoutModifier_Implementation(FName Category) const
+TArray<UWidgetStyleData*> UClassBasedWidgetStyleDefinitions::GetStylesForName_Implementation(FName Category) const
 {
 	const FWidgetStyleDataArray* Result = CategoriesWithoutModifier.Find(Category);
 	return Result
diff --git a/Engine/Plugins/Experimental/VirtualCamera/Source/VCamExtensions/Private/Styling/ModifierBoundWidgetStyleDefinitions.cpp b/Engine/Plugins/Experimental/VirtualCamera/Source/VCamExtensions/Private/Styling/ModifierBoundWidgetStyleDefinitions.cpp
--- a/Engine/Plugins/Experimental/VirtualCamera/Source/VCamExtensions/Private/Styling/ModifierBoundWidgetStyleDefinitions.cpp
+++ b/Engine/Plugins/Experimental/VirtualCamera/Source/VCamExtensions/Private/Styling/ModifierBoundWidgetStyleDefinitions.cpp
@@ -19,44 +19,44 @@ namespace UE::VCamExtensions::Private
 	}
 }
 
-TArray<UWidgetStyleData*> UModifierBoundWidgetStyleDefinitions::GetStylesForModifier_Implementation(UVCamComponent* VCamComponent, FName ModifierId) const
+TArray<UWidgetStyleData*> UModifierBoundWidgetStyleDefinitions::GetStylesForModifier_Implementation(UVCamModifier* Modifier) const
 {
 	unimplemented();
 	return {};
 }
 
-TArray<UWidgetStyleData*> UModifierBoundWidgetStyleDefinitions::GetStylesForCategoryInModifier_Implementation(UVCamComponent* VCamComponent, FName ModifierId, FName Category) const
+TArray<UWidgetStyleData*> UModifierBoundWidgetStyleDefinitions::GetStylesForConnectionPoint_Implementation(UVCamModifier* Modifier, FName ConnectionPoint) const
 {
 	unimplemented();
 	return {};
 }
 
-TArray<UWidgetStyleData*> UModifierBoundWidgetStyleDefinitions::GetStylesForCategoryWithoutModifier_Implementation(FName Category) const
+TArray<UWidgetStyleData*> UModifierBoundWidgetStyleDefinitions::GetStylesForName_Implementation(FName Category) const
 {
 	unimplemented();
 	return {};
 }
 
-UWidgetStyleData* UModifierBoundWidgetStyleDefinitions::GetStyleForModifierByClass(UVCamComponent* VCamComponent, FName ModifierId, TSubclassOf<UWidgetStyleData> Class) const
+UWidgetStyleData* UModifierBoundWidgetStyleDefinitions::GetStyleForModifierByClass(UVCamModifier* Modifier, TSubclassOf<UWidgetStyleData> Class) const
 {
 	return UE::VCamExtensions::Private::Filter(
-		GetStylesForModifier(VCamComponent, ModifierId),
+		GetStylesForModifier(Modifier),
 		Class
 		);
 }
 
-UWidgetStyleData* UModifierBoundWidgetStyleDefinitions::GetStyleForCategoryByClassInModifier(UVCamComponent* VCamComponent, FName ModifierId, FName Category, TSubclassOf<UWidgetStyleData> Class) const
+UWidgetStyleData* UModifierBoundWidgetStyleDefinitions::GetStyleForConnectionPointByClass(UVCamModifier* Modifier, FName ConnectionPoint, TSubclassOf<UWidgetStyleData> Class) const
 {
 	return UE::VCamExtensions::Private::Filter(
-		GetStylesForCategoryInModifier(VCamComponent, ModifierId, Category),
+		GetStylesForConnectionPoint(Modifier, ConnectionPoint),
 		Class
 		);
 }
 
-UWidgetStyleData* UModifierBoundWidgetStyleDefinitions::GetStyleForCategoryWithoutModifier(FName Category, TSubclassOf<UWidgetStyleData> Class) const
+UWidgetStyleData* UModifierBoundWidgetStyleDefinitions::GetStyleForNameByClass(FName Category, TSubclassOf<UWidgetStyleData> Class) const
 {
 	return UE::VCamExtensions::Private::Filter(
-			GetStylesForCategoryWithoutModifier(Category),
-			Class
-			);
+		GetStylesForName(Category),
+		Class
+		);
 }
\ No newline at end of file
diff --git a/Engine/Plugins/Experimental/VirtualCamera/Source/VCamExtensions/Private/Styling/ModifierBoundWidgetStylesAsset.cpp b/Engine/Plugins/Experimental/VirtualCamera/Source/VCamExtensions/Private/Styling/ModifierBoundWidgetStylesAsset.cpp
--- a/Engine/Plugins/Experimental/VirtualCamera/Source/VCamExtensions/Private/Styling/ModifierBoundWidgetStylesAsset.cpp
+++ b/Engine/Plugins/Experimental/VirtualCamera/Source/VCamExtensions/Private/Styling/ModifierBoundWidgetStylesAsset.cpp
@@ -5,32 +5,32 @@
 #include "Styling/ClassBasedWidgetStyleDefinitions.h"
 #include "Styling/ModifierBoundWidgetStyleDefinitions.h"
 
-TArray<UWidgetStyleData*> UModifierBoundWidgetStylesAsset::GetStylesForModifier(UVCamComponent* VCamComponent, FName ModifierId) const
+TArray<UWidgetStyleData*> UModifierBoundWidgetStylesAsset::GetStylesForModifier(UVCamModifier* Modifier) const
 {
-	return Rules ? Rules->GetStylesForModifier(VCamComponent, ModifierId) : TArray<UWidgetStyleData*>{};
+	return Rules ? Rules->GetStylesForModifier(Modifier) : TArray<UWidgetStyleData*>{};
 }
 
-TArray<UWidgetStyleData*> UModifierBoundWidgetStylesAsset::GetStylesForCategoryInModifier(UVCamComponent* VCamComponent, FName ModifierId, FName Category) const
+TArray<UWidgetStyleData*> UModifierBoundWidgetStylesAsset::GetStylesForConnectionPoint(UVCamModifier* Modifier, FName ConnectionPoint) const
 {
-	return Rules ? Rules->GetStylesForCategoryInModifier(VCamComponent, ModifierId, Category) : TArray<UWidgetStyleData*>{};
+	return Rules ? Rules->GetStylesForConnectionPoint(Modifier, ConnectionPoint) : TArray<UWidgetStyleData*>{};
 }
 
-TArray<UWidgetStyleData*> UModifierBoundWidgetStylesAsset::GetStylesForCategoryWithoutModifier(FName Category) const
+TArray<UWidgetStyleData*> UModifierBoundWidgetStylesAsset::GetStylesForName(FName Category) const
 {
-	return Rules ? Rules->GetStylesForCategoryWithoutModifier(Category) : TArray<UWidgetStyleData*>{};
+	return Rules ? Rules->GetStylesForName(Category) : TArray<UWidgetStyleData*>{};
 }
 
-UWidgetStyleData* UModifierBoundWidgetStylesAsset::GetStyleForModifierByClass(UVCamComponent* VCamComponent, FName ModifierId, TSubclassOf<UWidgetStyleData> Class) const
+UWidgetStyleData* UModifierBoundWidgetStylesAsset::GetStyleForModifierByClass(UVCamModifier* Modifier, TSubclassOf<UWidgetStyleData> Class) const
 {
-	return Rules ? Rules->GetStyleForModifierByClass(VCamComponent, ModifierId, Class) : nullptr;
+	return Rules ? Rules->GetStyleForModifierByClass(Modifier, Class) : nullptr;
 }
 
-UWidgetStyleData* UModifierBoundWidgetStylesAsset::GetStyleForCategoryByClassInModifier(UVCamComponent* VCamComponent, FName ModifierId, FName Category, TSubclassOf<UWidgetStyleData> Class) const
+UWidgetStyleData* UModifierBoundWidgetStylesAsset::GetStyleForConnectionPointByClass(UVCamModifier* Modifier, FName ConnectionPoint, TSubclassOf<UWidgetStyleData> Class) const
 {
-	return Rules ? Rules->GetStyleForCategoryByClassInModifier(VCamComponent, ModifierId, Category, Class) : nullptr;
+	return Rules ? Rules->GetStyleForConnectionPointByClass(Modifier, ConnectionPoint, Class) : nullptr;
 }
 
-UWidgetStyleData* UModifierBoundWidgetStylesAsset::GetStyleForCategoryWithoutModifier(FName Category, TSubclassOf<UWidgetStyleData> Class) const
+UWidgetStyleData* UModifierBoundWidgetStylesAsset::GetStyleForNameByClass(FName Name, TSubclassOf<UWidgetStyleData> Class) const
 {
-	return Rules ? Rules->GetStyleForCategoryWithoutModifier(Category, Class) : nullptr;
+	return Rules ? Rules->GetStyleForNameByClass(Name, Class) : nullptr;
 }
diff --git a/Engine/Plugins/Experimental/VirtualCamera/Source/VCamExtensions/Private/Styling/TargetBasedWidgetStyleDefinitions.cpp b/Engine/Plugins/Experimental/VirtualCamera/Source/VCamExtensions/Private/Styling/TargetBasedWidgetStyleDefinitions.cpp
new file mode 100644
--- /dev/null
+++ b/Engine/Plugins/Experimental/VirtualCamera/Source/VCamExtensions/Private/Styling/TargetBasedWidgetStyleDefinitions.cpp
@@ -0,0 +1,55 @@
+// Copyright Epic Games, Inc. All Rights Reserved.
+
+#include "Styling/TargetBasedWidgetStyleDefinitions.h"
+
+#include "Modifier/VCamModifier.h"
+
+TArray<UWidgetStyleData*> UTargetBasedWidgetStyleDefinitions::GetStylesForModifier_Implementation(UVCamModifier* Modifier) const
+{
+	if (!Modifier)
+	{
+		return {};
+	}
+	
+	if (const FTargettedModifierStyleConfig* Data = ModifierToStyle.Find(Modifier->GetStackEntryName()))
+	{
+		TArray<UWidgetStyleData*> Result;
+		Algo::Transform(Data->ModifierStyles.Styles, Result, [](const TObjectPtr<UWidgetStyleData>& Data){ return Data; });
+		return Result;
+	}
+	return {};
+}
+
+TArray<UWidgetStyleData*> UTargetBasedWidgetStyleDefinitions::GetStylesForConnectionPoint_Implementation(UVCamModifier* Modifier, FName ConnectionPointId) const
+{
+	if (!Modifier)
+	{
+		return {};
+	}
+	
+	const FTargettedModifierStyleConfig* Data = ModifierToStyle.Find(Modifier->GetStackEntryName());
+	if (!Data)
+	{
+		return {};
+	}
+
+	if (const FWidgetStyleDataArray* Styles = Data->ConnectionPointStyles.Find(ConnectionPointId))
+	{
+		TArray<UWidgetStyleData*> Result;
+		Algo::Transform(Styles->Styles, Result, [](const TObjectPtr<UWidgetStyleData>& Data){ return Data; });
+		return Result;
+	}
+	
+	return {};
+}
+
+TArray<UWidgetStyleData*> UTargetBasedWidgetStyleDefinitions::GetStylesForName_Implementation(FName Category) const
+{
+	if (const FWidgetStyleDataArray* Data = CategoriesWithoutModifier.Find(Category))
+	{
+		TArray<UWidgetStyleData*> Result;
+		Algo::Transform(Data->Styles, Result, [](const TObjectPtr<UWidgetStyleData>& Data){ return Data; });
+		return Result;
+	}
+	return {};
+}
\ No newline at end of file
diff --git a/Engine/Plugins/Experimental/VirtualCamera/Source/VCamExtensions/Public/Styling/ClassBasedWidgetStyleDefinitions.h b/Engine/Plugins/Experimental/VirtualCamera/Source/VCamExtensions/Public/Styling/ClassBasedWidgetStyleDefinitions.h
--- a/Engine/Plugins/Experimental/VirtualCamera/Source/VCamExtensions/Public/Styling/ClassBasedWidgetStyleDefinitions.h
+++ b/Engine/Plugins/Experimental/VirtualCamera/Source/VCamExtensions/Public/Styling/ClassBasedWidgetStyleDefinitions.h
@@ -4,19 +4,11 @@
 
 #include "ModifierBoundWidgetStyleDefinitions.h"
 #include "Modifier/VCamModifier.h"
+#include "WidgetStyleDataArray.h"
 #include "ClassBasedWidgetStyleDefinitions.generated.h"
 
 class UWidgetStyleData;
 
-USTRUCT()
-struct VCAMEXTENSIONS_API FWidgetStyleDataArray
-{
-	GENERATED_BODY()
-	
-	UPROPERTY(EditAnywhere, Instanced, Category = "Virtual Camera")
-	TArray<TObjectPtr<UWidgetStyleData>> Styles;
-};
-
 USTRUCT()
 struct VCAMEXTENSIONS_API FWidgetStyleDataConfig
 {
@@ -41,7 +33,7 @@ struct VCAMEXTENSIONS_API FPerModifierClassWidgetSytleData
 
 	/** Meta data that is attached to a custom name, such as a connection point or group name (if using e.g. UClassBasedModifierHierarchyRules) */
 	UPROPERTY(EditAnywhere, Category = "Virtual Camera")
-	TMap<FName, FWidgetStyleDataConfig> CategorizedStyles;
+	TMap<FName, FWidgetStyleDataConfig> ConnectionPointStyles;
 };
 
 /**
@@ -61,8 +53,8 @@ public:
 	TMap<FName, FWidgetStyleDataArray> CategoriesWithoutModifier;
 	
 	//~ Begin UModifierMetaDataRules Interface
-	virtual TArray<UWidgetStyleData*> GetStylesForModifier_Implementation(UVCamComponent* VCamComponent, FName ModifierId) const override;
-	virtual TArray<UWidgetStyleData*> GetStylesForCategoryInModifier_Implementation(UVCamComponent* VCamComponent, FName ModifierId, FName ConnectionPointId) const override;
-	virtual TArray<UWidgetStyleData*> GetStylesForCategoryWithoutModifier_Implementation(FName Category) const override;
+	virtual TArray<UWidgetStyleData*> GetStylesForModifier_Implementation(UVCamModifier* Modifier) const override;
+	virtual TArray<UWidgetStyleData*> GetStylesForConnectionPoint_Implementation(UVCamModifier* Modifier, FName ConnectionPointId) const override;
+	virtual TArray<UWidgetStyleData*> GetStylesForName_Implementation(FName Category) const override;
 	//~ End UModifierMetaDataRules Interface
 };
diff --git a/Engine/Plugins/Experimental/VirtualCamera/Source/VCamExtensions/Public/Styling/ModifierBoundWidgetStyleDefinitions.h b/Engine/Plugins/Experimental/VirtualCamera/Source/VCamExtensions/Public/Styling/ModifierBoundWidgetStyleDefinitions.h
--- a/Engine/Plugins/Experimental/VirtualCamera/Source/VCamExtensions/Public/Styling/ModifierBoundWidgetStyleDefinitions.h
+++ b/Engine/Plugins/Experimental/VirtualCamera/Source/VCamExtensions/Public/Styling/ModifierBoundWidgetStyleDefinitions.h
@@ -8,7 +8,7 @@
 #include "ModifierBoundWidgetStyleDefinitions.generated.h"
 
 class UWidgetStyleData;
-class UVCamComponent;
+class UVCamModifier;
 
 /**
  * Associates information with modifiers and their connection points.
@@ -23,22 +23,22 @@ public:
 
 	/** Retrieves all meta data that is associated for a given modifier. */
 	UFUNCTION(BlueprintPure, BlueprintNativeEvent, Category = "Virtual Camera")
-	TArray<UWidgetStyleData*> GetStylesForModifier(UVCamComponent* VCamComponent, FName ModifierId) const;
+	TArray<UWidgetStyleData*> GetStylesForModifier(UVCamModifier* Modifier) const;
 
 	/** Retrieves all meta data that is associated for a given modifier and a sub-category name. */
 	UFUNCTION(BlueprintPure, BlueprintNativeEvent, Category = "Virtual Camera")
-	TArray<UWidgetStyleData*> GetStylesForCategoryInModifier(UVCamComponent* VCamComponent, FName ModifierId, FName Category) const;
+	TArray<UWidgetStyleData*> GetStylesForConnectionPoint(UVCamModifier* Modifier, FName ConnectionPoint) const;
 
 	/** Retrieves all meta data that is associated with a given a category name; this data is not associated with any kind of modifier. */
 	UFUNCTION(BlueprintPure, BlueprintNativeEvent, Category = "Virtual Camera")
-	TArray<UWidgetStyleData*> GetStylesForCategoryWithoutModifier(FName Category) const;
+	TArray<UWidgetStyleData*> GetStylesForName(FName Category) const;
 
 	UFUNCTION(BlueprintPure, Category = "Virtual Camera", meta = (DeterminesOutputType = "Class"))
-	UWidgetStyleData* GetStyleForModifierByClass(UVCamComponent* VCamComponent, FName ModifierId, TSubclassOf<UWidgetStyleData> Class) const;
+	UWidgetStyleData* GetStyleForModifierByClass(UVCamModifier* Modifier, TSubclassOf<UWidgetStyleData> Class) const;
 	
 	UFUNCTION(BlueprintPure, Category = "Virtual Camera", meta = (DeterminesOutputType = "Class"))
-	UWidgetStyleData* GetStyleForCategoryByClassInModifier(UVCamComponent* VCamComponent, FName ModifierId, FName Category, TSubclassOf<UWidgetStyleData> Class) const;
+	UWidgetStyleData* GetStyleForConnectionPointByClass(UVCamModifier* Modifier, FName ConnectionPoint, TSubclassOf<UWidgetStyleData> Class) const;
 	
 	UFUNCTION(BlueprintPure, Category = "Virtual Camera", meta = (DeterminesOutputType = "Class"))
-	UWidgetStyleData* GetStyleForCategoryWithoutModifier( FName Category, TSubclassOf<UWidgetStyleData> Class) const;
+	UWidgetStyleData* GetStyleForNameByClass(FName Category, TSubclassOf<UWidgetStyleData> Class) const;
 };
diff --git a/Engine/Plugins/Experimental/VirtualCamera/Source/VCamExtensions/Public/Styling/ModifierBoundWidgetStylesAsset.h b/Engine/Plugins/Experimental/VirtualCamera/Source/VCamExtensions/Public/Styling/ModifierBoundWidgetStylesAsset.h
--- a/Engine/Plugins/Experimental/VirtualCamera/Source/VCamExtensions/Public/Styling/ModifierBoundWidgetStylesAsset.h
+++ b/Engine/Plugins/Experimental/VirtualCamera/Source/VCamExtensions/Public/Styling/ModifierBoundWidgetStylesAsset.h
@@ -23,22 +23,22 @@ public:
 
 	/** Retrieves all meta data that is associated for a given modifier. */
 	UFUNCTION(BlueprintPure, Category = "Virtual Camera")
-	TArray<UWidgetStyleData*> GetStylesForModifier(UVCamComponent* VCamComponent, FName ModifierId) const;
+	TArray<UWidgetStyleData*> GetStylesForModifier(UVCamModifier* Modifier) const;
 
 	/** Retrieves all meta data that is associated for a given modifier and a sub-category name. */
 	UFUNCTION(BlueprintPure, Category = "Virtual Camera")
-	TArray<UWidgetStyleData*> GetStylesForCategoryInModifier(UVCamComponent* VCamComponent, FName ModifierId, FName Category) const;
+	TArray<UWidgetStyleData*> GetStylesForConnectionPoint(UVCamModifier* Modifier, FName ConnectionPoint) const;
 
 	/** Retrieves all meta data that is associated with a given a category name; this data is not associated with any kind of modifier. */
 	UFUNCTION(BlueprintPure, Category = "Virtual Camera")
-	TArray<UWidgetStyleData*> GetStylesForCategoryWithoutModifier(FName Category) const;
+	TArray<UWidgetStyleData*> GetStylesForName(FName Category) const;
 
 	UFUNCTION(BlueprintPure, Category = "Virtual Camera", meta = (DeterminesOutputType = "Class"))
-	UWidgetStyleData* GetStyleForModifierByClass(UVCamComponent* VCamComponent, FName ModifierId, TSubclassOf<UWidgetStyleData> Class) const;
+	UWidgetStyleData* GetStyleForModifierByClass(UVCamModifier* Modifier, TSubclassOf<UWidgetStyleData> Class) const;
 	
 	UFUNCTION(BlueprintPure, Category = "Virtual Camera", meta = (DeterminesOutputType = "Class"))
-	UWidgetStyleData* GetStyleForCategoryByClassInModifier(UVCamComponent* VCamComponent, FName ModifierId, FName Category, TSubclassOf<UWidgetStyleData> Class) const;
+	UWidgetStyleData* GetStyleForConnectionPointByClass(UVCamModifier* Modifier, FName ConnectionPoint, TSubclassOf<UWidgetStyleData> Class) const;
 	
 	UFUNCTION(BlueprintPure, Category = "Virtual Camera", meta = (DeterminesOutputType = "Class"))
-	UWidgetStyleData* GetStyleForCategoryWithoutModifier( FName Category, TSubclassOf<UWidgetStyleData> Class) const;
+	UWidgetStyleData* GetStyleForNameByClass(FName Name, TSubclassOf<UWidgetStyleData> Class) const;
 };
\ No newline at end of file
diff --git a/Engine/Plugins/Experimental/VirtualCamera/Source/VCamExtensions/Public/Styling/TargetBasedWidgetStyleDefinitions.h b/Engine/Plugins/Experimental/VirtualCamera/Source/VCamExtensions/Public/Styling/TargetBasedWidgetStyleDefinitions.h
new file mode 100644
--- /dev/null
+++ b/Engine/Plugins/Experimental/VirtualCamera/Source/VCamExtensions/Public/Styling/TargetBasedWidgetStyleDefinitions.h
@@ -0,0 +1,41 @@
+// Copyright Epic Games, Inc. All Rights Reserved.
+
+#pragma once
+
+#include "ModifierBoundWidgetStyleDefinitions.h"
+#include "WidgetStyleDataArray.h"
+#include "TargetBasedWidgetStyleDefinitions.generated.h"
+
+USTRUCT()
+struct VCAMEXTENSIONS_API FTargettedModifierStyleConfig
+{
+	GENERATED_BODY()
+
+	UPROPERTY(EditAnywhere, Category = "Virtual Camera")
+	FWidgetStyleDataArray ModifierStyles;
+
+	/** Key: Name of a connection point on the modifier */
+	UPROPERTY(EditAnywhere, Category = "Virtual Camera")
+	TMap<FName, FWidgetStyleDataArray> ConnectionPointStyles;
+};
+
+UCLASS()
+class VCAMEXTENSIONS_API UTargetBasedWidgetStyleDefinitions : public UModifierBoundWidgetStyleDefinitions
+{
+	GENERATED_BODY()
+public:
+
+	/** Key: Name of a modifier in a VCam's stack entry. */
+	UPROPERTY(EditAnywhere, Category = "Virtual Camera")
+	TMap<FName, FTargettedModifierStyleConfig> ModifierToStyle;
+
+	/** Styles that is not bound to any modifier but just to a name. */
+	UPROPERTY(EditAnywhere, Category = "Virtual Camera")
+	TMap<FName, FWidgetStyleDataArray> CategoriesWithoutModifier;
+
+	//~ Begin UModifierMetaDataRules Interface
+	virtual TArray<UWidgetStyleData*> GetStylesForModifier_Implementation(UVCamModifier* Modifier) const override;
+	virtual TArray<UWidgetStyleData*> GetStylesForConnectionPoint_Implementation(UVCamModifier* Modifier, FName ConnectionPointId) const override;
+	virtual TArray<UWidgetStyleData*> GetStylesForName_Implementation(FName Category) const override;
+	//~ End UModifierMetaDataRules Interface
+};
diff --git a/Engine/Plugins/Experimental/VirtualCamera/Source/VCamExtensions/Public/Styling/WidgetStyleDataArray.h b/Engine/Plugins/Experimental/VirtualCamera/Source/VCamExtensions/Public/Styling/WidgetStyleDataArray.h
new file mode 100644
--- /dev/null
+++ b/Engine/Plugins/Experimental/VirtualCamera/Source/VCamExtensions/Public/Styling/WidgetStyleDataArray.h
@@ -0,0 +1,16 @@
+// Copyright Epic Games, Inc. All Rights Reserved.
+
+#pragma once
+
+#include "WidgetStyleDataArray.generated.h"
+
+class UWidgetStyleData;
+
+USTRUCT()
+struct VCAMEXTENSIONS_API FWidgetStyleDataArray
+{
+	GENERATED_BODY()
+	
+	UPROPERTY(EditAnywhere, Instanced, Category = "Virtual Camera")
+	TArray<TObjectPtr<UWidgetStyleData>> Styles;
+};
\ No newline at end of file
