Fix: Complete integration work and force media capture to be completed on render thread

diff --git a/Engine/Plugins/Experimental/VirtualCamera/Source/PixelStreamingVCam/Private/PixelStreamingMediaCapture.cpp b/Engine/Plugins/Experimental/VirtualCamera/Source/PixelStreamingVCam/Private/PixelStreamingMediaCapture.cpp
--- a/Engine/Plugins/Experimental/VirtualCamera/Source/PixelStreamingVCam/Private/PixelStreamingMediaCapture.cpp
+++ b/Engine/Plugins/Experimental/VirtualCamera/Source/PixelStreamingVCam/Private/PixelStreamingMediaCapture.cpp
@@ -23,6 +23,13 @@ bool UPixelStreamingMediaCapture::InitializeCapture()
 	UE_LOG(LogPixelStreamingVCam, Log, TEXT("Initializing media capture for Pixel Streaming VCam."));
 	bViewportResized = false;
 	SetState(EMediaCaptureState::Capturing);
+
+	// Force the MediaCapture readback to be completed on the render thread
+	IConsoleVariable* CVar = IConsoleManager::Get().FindConsoleVariable(TEXT("MediaIO.ScheduleOnAnyThread"));
+	if(CVar) {
+		CVar->Set(0, EConsoleVariableFlags::ECVF_SetByCode);
+	}
+
 	return true;
 }
 
diff --git a/Engine/Plugins/Experimental/VirtualCamera/Source/PixelStreamingVCam/Private/PixelStreamingMediaOutput.cpp b/Engine/Plugins/Experimental/VirtualCamera/Source/PixelStreamingVCam/Private/PixelStreamingMediaOutput.cpp
--- a/Engine/Plugins/Experimental/VirtualCamera/Source/PixelStreamingVCam/Private/PixelStreamingMediaOutput.cpp
+++ b/Engine/Plugins/Experimental/VirtualCamera/Source/PixelStreamingVCam/Private/PixelStreamingMediaOutput.cpp
@@ -46,28 +46,27 @@ UMediaCapture* UPixelStreamingMediaOutput::CreateMediaCaptureImpl()
 
 void UPixelStreamingMediaOutput::RegisterRemoteResolutionCommandHandler()
 {
-	// TODO: Complete integration work from dev stream
 	// Override resolution command as we this to set the output provider override resolution
-	// TSharedPtr<IPixelStreamingInputHandler> InputHandler = Streamer->GetInputHandler().Pin();
-	// if(InputHandler)
-	// {
-	// 	InputHandler->SetCommandHandler(TEXT("Resolution.Width"), [this](FString Descriptor, FString WidthString){
-	// 		bool bSuccess = false;
-	// 		FString HeightString;
-	// 		UE::PixelStreaming::ExtractJsonFromDescriptor(Descriptor, TEXT("Resolution.Height"), HeightString, bSuccess);
-	// 		if (bSuccess)
-	// 		{
-	// 			int Width = FCString::Atoi(*WidthString);
-	// 			int Height = FCString::Atoi(*HeightString);
-	// 			if (Width < 1 || Height < 1)
-	// 			{
-	// 				return;
-	// 			}
+	TSharedPtr<IPixelStreamingInputHandler> InputHandler = Streamer->GetInputHandler().Pin();
+	if(InputHandler)
+	{
+		InputHandler->SetCommandHandler(TEXT("Resolution.Width"), [this](FString Descriptor, FString WidthString){
+			bool bSuccess = false;
+			FString HeightString;
+			UE::PixelStreaming::ExtractJsonFromDescriptor(Descriptor, TEXT("Resolution.Height"), HeightString, bSuccess);
+			if (bSuccess)
+			{
+				int Width = FCString::Atoi(*WidthString);
+				int Height = FCString::Atoi(*HeightString);
+				if (Width < 1 || Height < 1)
+				{
+					return;
+				}
 
-	// 			RemoteResolutionChangedEvent.Broadcast(FIntPoint(Width, Height));
-	// 		}
-	// 	});
-	// }
+				RemoteResolutionChangedEvent.Broadcast(FIntPoint(Width, Height));
+			}
+		});
+	}
 }
 
 void UPixelStreamingMediaOutput::StartStreaming()
