Fix non unity issues & minimize includes

diff --git a/Engine/Plugins/Experimental/VirtualCamera/Source/VCamCore/Private/VCamComponent.cpp b/Engine/Plugins/Experimental/VirtualCamera/Source/VCamCore/Private/VCamComponent.cpp
--- a/Engine/Plugins/Experimental/VirtualCamera/Source/VCamCore/Private/VCamComponent.cpp
+++ b/Engine/Plugins/Experimental/VirtualCamera/Source/VCamCore/Private/VCamComponent.cpp
@@ -12,16 +12,16 @@
 #include "CineCameraComponent.h"
 #include "Engine/GameEngine.h"
 #include "Engine/InputDelegateBinding.h"
-#include "EnhancedInputComponent.h"
-#include "EnhancedInputSubsystems.h"
+#include "EnhancedActionKeyMapping.h"
+#include "EnhancedInputSubsystemInterface.h"
 #include "Features/IModularFeatures.h"
 #include "GameFramework/InputSettings.h"
 #include "ILiveLinkClient.h"
 #include "InputMappingContext.h"
-#include "PlayerMappableKeySettings.h"
 #include "Input/InputVCamSubsystem.h"
 #include "Roles/LiveLinkCameraRole.h"
 #include "Roles/LiveLinkTransformRole.h"
+#include "VCamSubsystem.h"
 
 #if WITH_EDITOR
 #include "Modules/ModuleManager.h"
@@ -29,13 +29,12 @@
 #include "LevelEditor.h"
 #include "SLevelViewport.h"
 
-#include "IConcertModule.h"
 #include "IConcertClient.h"
 #include "IConcertSession.h"
+#include "IConcertSessionHandler.h"
 #include "IConcertSyncClient.h"
 #include "IMultiUserClientModule.h"
-
-#include "EnhancedInputEditorSubsystem.h"
+#include "VCamMultiUser.h"
 
 #include "VPSettings.h"
 #include "VPRolesSubsystem.h"
diff --git a/Engine/Plugins/Experimental/VirtualCamera/Source/VCamCore/Public/VCamComponent.h b/Engine/Plugins/Experimental/VirtualCamera/Source/VCamCore/Public/VCamComponent.h
--- a/Engine/Plugins/Experimental/VirtualCamera/Source/VCamCore/Public/VCamComponent.h
+++ b/Engine/Plugins/Experimental/VirtualCamera/Source/VCamCore/Public/VCamComponent.h
@@ -7,37 +7,40 @@
 #include "EVCamTargetViewportID.h"
 #include "Modifier/ModifierStackEntry.h"
 
-#include "EnhancedActionKeyMapping.h"
-#include "EnhancedInputSubsystemInterface.h"
 #include "GameplayTagContainer.h"
 #include "Roles/LiveLinkCameraTypes.h"
 #include "VCamInputSettings.h"
 #include "VCamSubsystem.h"
 #include "Input/VCamInputDeviceConfig.h"
 #include "Util/VCamViewportLocker.h"
+#include "Subsystems/SubsystemCollection.h"
 
 #if WITH_EDITOR
-#include "Engine/GameEngine.h"
 #include "UnrealEdMisc.h"
-#include "VCamMultiUser.h"
 #endif
 
 #include "VCamComponent.generated.h"
 
-class UInputVCamSubsystem;
-class UVCamOutputProviderBase;
 DECLARE_LOG_CATEGORY_EXTERN(LogVCamComponent, Log, All);
 
+class IEnhancedInputSubsystemInterface;
 class UCineCameraComponent;
+class UInputModifier;
+class UInputTrigger;
+class UInputVCamSubsystem;
 class UVCamModifierContext;
-class SWindow;
-class FSceneViewport;
-class UEnhancedInputComponent;
+class UVCamOutputProviderBase;
+
+struct FEnhancedActionKeyMapping;
 
 #if WITH_EDITOR
 class FLevelEditorViewportClient;
+class IConcertClientSession;
+struct FConcertSessionContext;
+struct FMultiUserVCamCameraComponentEvent;
 #endif
 
+
 DECLARE_DYNAMIC_MULTICAST_DELEGATE_OneParam(FOnComponentReplaced, UVCamComponent*, NewComponent);
 
 /**
