Fix non-unity error on Mac

diff --git a/Engine/Plugins/Experimental/VirtualCamera/Source/VCamCore/Private/UI/VCamConnectionStructs.cpp b/Engine/Plugins/Experimental/VirtualCamera/Source/VCamCore/Private/UI/VCamConnectionStructs.cpp
--- a/Engine/Plugins/Experimental/VirtualCamera/Source/VCamCore/Private/UI/VCamConnectionStructs.cpp
+++ b/Engine/Plugins/Experimental/VirtualCamera/Source/VCamCore/Private/UI/VCamConnectionStructs.cpp
@@ -2,9 +2,11 @@
 
 #include "UI/VCamConnectionStructs.h"
 
-#include "Algo/AllOf.h"
 #include "VCamComponent.h"
 
+#include "Algo/AllOf.h"
+#include "InputAction.h"
+
 DEFINE_LOG_CATEGORY(LogVCamConnection);
 
 bool FVCamConnection::IsConnected() const
