Fixed assumption that each modifier would be referenced by at most one modifier and improved API

diff --git a/Engine/Plugins/Experimental/VirtualCamera/Source/VCamExtensions/Private/Hierarchies/BaseModifierGroup.cpp b/Engine/Plugins/Experimental/VirtualCamera/Source/VCamExtensions/Private/Hierarchies/BaseModifierGroup.cpp
--- a/Engine/Plugins/Experimental/VirtualCamera/Source/VCamExtensions/Private/Hierarchies/BaseModifierGroup.cpp
+++ b/Engine/Plugins/Experimental/VirtualCamera/Source/VCamExtensions/Private/Hierarchies/BaseModifierGroup.cpp
@@ -2,20 +2,10 @@
 
 #include "Hierarchies/BaseModifierGroup.h"
 
+#include "Misc/Guid.h"
+
 void UBaseModifierGroup::PostInitProperties()
 {
-	GroupName = GetFName();
+	NodeName = FName(*FGuid::NewGuid().ToString());
 	UObject::PostInitProperties();
-}
-
-#if WITH_EDITOR
-void UBaseModifierGroup::PostEditChangeProperty(FPropertyChangedEvent& PropertyChangedEvent)
-{
-	if (PropertyChangedEvent.Property && PropertyChangedEvent.Property->GetFName() == GET_MEMBER_NAME_CHECKED(UBaseModifierGroup, GroupName))
-	{
-		// TODO:
-	}
-	
-	UObject::PostEditChangeProperty(PropertyChangedEvent);
-}
-#endif
+}
\ No newline at end of file
diff --git a/Engine/Plugins/Experimental/VirtualCamera/Source/VCamExtensions/Private/Hierarchies/ClassBasedModifierHierarchyRules.cpp b/Engine/Plugins/Experimental/VirtualCamera/Source/VCamExtensions/Private/Hierarchies/ClassBasedModifierHierarchyRules.cpp
deleted file mode 100644
diff --git a/Engine/Plugins/Experimental/VirtualCamera/Source/VCamExtensions/Private/Hierarchies/ModifierHierarchyAsset.cpp b/Engine/Plugins/Experimental/VirtualCamera/Source/VCamExtensions/Private/Hierarchies/ModifierHierarchyAsset.cpp
--- a/Engine/Plugins/Experimental/VirtualCamera/Source/VCamExtensions/Private/Hierarchies/ModifierHierarchyAsset.cpp
+++ b/Engine/Plugins/Experimental/VirtualCamera/Source/VCamExtensions/Private/Hierarchies/ModifierHierarchyAsset.cpp
@@ -2,25 +2,42 @@
 
 #include "Hierarchies/ModifierHierarchyAsset.h"
 
-#include "Hierarchies/ClassBasedModifierHierarchyRules.h"
 #include "Hierarchies/ModifierHierarchyRules.h"
 
-FName UModifierHierarchyAsset::GetRootGroup() const
+FName UModifierHierarchyAsset::GetRootNode() const
 {
-	return Rules ? Rules->GetRootGroup() : NAME_None;
+	return Rules ? Rules->GetRootNode() : NAME_None;
 }
 
-bool UModifierHierarchyAsset::GetGroupOfModifier(UVCamModifier* Modifier, FName& Group) const
+bool UModifierHierarchyAsset::GetParentNode(FName ChildNode, FName& ParentNode) const
 {
-	return Rules && Rules->GetGroupOfModifier(Modifier, Group);
+	return Rules && Rules->GetParentNode(ChildNode, ParentNode);
 }
 
-TSet<FName> UModifierHierarchyAsset::GetChildGroups_Implementation(FName ParentGroup) const
+TSet<FName> UModifierHierarchyAsset::GetChildNodes(FName ParentGroup) const
 {
-	return Rules ? Rules->GetChildGroups(ParentGroup) : TSet<FName>{};
+	return Rules ? Rules->GetChildNodes(ParentGroup) : TSet<FName>{};
 }
 
-TSet<UVCamModifier*> UModifierHierarchyAsset::GetModifiersInGroup(UVCamComponent* Component, FName GroupName) const
+UVCamModifier* UModifierHierarchyAsset::GetModifierInNode(UVCamComponent* Component, FName GroupName) const
 {
-	return Rules ? Rules->GetModifiersInGroup(Component, GroupName) : TSet<UVCamModifier*>{};
+	return Rules ? Rules->GetModifierInNode(Component, GroupName) : nullptr;
+}
+
+bool UModifierHierarchyAsset::GetConnectionPointTargetForNode(FName GroupName, UVCamComponent* Component, FVCamModifierConnectionBinding& Connection) const
+{
+	if (Rules)
+	{
+		return Rules->GetConnectionPointTargetForNode(GroupName, Component, Connection);
+	}
+	return false;
+}
+
+TSet<FName> UModifierHierarchyAsset::GetGroupsContainingModifier(UVCamModifier* Modifier) const
+{
+	if (Rules)
+	{
+		return Rules->GetNodesContainingModifier(Modifier);
+	}
+	return {};
 }
diff --git a/Engine/Plugins/Experimental/VirtualCamera/Source/VCamExtensions/Private/Hierarchies/ModifierHierarchyRules.cpp b/Engine/Plugins/Experimental/VirtualCamera/Source/VCamExtensions/Private/Hierarchies/ModifierHierarchyRules.cpp
--- a/Engine/Plugins/Experimental/VirtualCamera/Source/VCamExtensions/Private/Hierarchies/ModifierHierarchyRules.cpp
+++ b/Engine/Plugins/Experimental/VirtualCamera/Source/VCamExtensions/Private/Hierarchies/ModifierHierarchyRules.cpp
@@ -2,38 +2,36 @@
 
 #include "Hierarchies/ModifierHierarchyRules.h"
 
-#include "VCamExtensionsLog.h"
-
-FName UModifierHierarchyRules::GetRootGroup_Implementation() const
+FName UModifierHierarchyRules::GetRootNode_Implementation() const
 {
 	unimplemented();
 	return {};
 }
 
-bool UModifierHierarchyRules::GetParentGroup_Implementation(FName ChildGroup, FName& ParentGroup) const
+bool UModifierHierarchyRules::GetParentNode_Implementation(FName ChildNode, FName& ParentNode) const
 {
 	return false;
 }
 
-bool UModifierHierarchyRules::GetGroupOfModifier_Implementation(UVCamModifier* Modifier, FName& Group) const
+TSet<FName> UModifierHierarchyRules::GetChildNodes_Implementation(FName Node) const
 {
 	unimplemented();
-	return false;
+	return {};
 }
 
-bool UModifierHierarchyRules::GetConnectionPointForModifier_Implementation(UVCamModifier* Modifier, FName& ConnectionPoint) const
+UVCamModifier* UModifierHierarchyRules::GetModifierInNode_Implementation(UVCamComponent* Component, FName NodeName) const
 {
-	UE_LOG(LogVCamExtensions, Warning, TEXT("%s does not support GetConnectionPointForModifier"), *GetClass()->GetName());
-	return false;
+	unimplemented();
+	return {};
 }
 
-TSet<FName> UModifierHierarchyRules::GetChildGroups_Implementation(FName ParentGroup) const
+bool UModifierHierarchyRules::GetConnectionPointTargetForNode_Implementation(FName GroupName, UVCamComponent* Component, FVCamModifierConnectionBinding& Connection) const
 {
 	unimplemented();
-	return {};
+	return false;
 }
 
-TSet<UVCamModifier*> UModifierHierarchyRules::GetModifiersInGroup_Implementation(UVCamComponent* Component, FName GroupName) const
+TSet<FName> UModifierHierarchyRules::GetNodesContainingModifier_Implementation(UVCamModifier* Modifier) const
 {
 	unimplemented();
 	return {};
diff --git a/Engine/Plugins/Experimental/VirtualCamera/Source/VCamExtensions/Private/Hierarchies/TargetModifierPerGroupHierarchyRules.cpp b/Engine/Plugins/Experimental/VirtualCamera/Source/VCamExtensions/Private/Hierarchies/TargetModifierPerGroupHierarchyRules.cpp
deleted file mode 100644
diff --git a/Engine/Plugins/Experimental/VirtualCamera/Source/VCamExtensions/Private/Hierarchies/TargetModifierPerNodeHierarchyRules.cpp b/Engine/Plugins/Experimental/VirtualCamera/Source/VCamExtensions/Private/Hierarchies/TargetModifierPerNodeHierarchyRules.cpp
new file mode 100644
--- /dev/null
+++ b/Engine/Plugins/Experimental/VirtualCamera/Source/VCamExtensions/Private/Hierarchies/TargetModifierPerNodeHierarchyRules.cpp
@@ -0,0 +1,134 @@
+// Copyright Epic Games, Inc. All Rights Reserved.
+
+#include "Hierarchies/TargetModifierPerNodeHierarchyRules.h"
+
+#include "HierarchyUtils.h"
+#include "VCamComponent.h"
+#include "VCamExtensionsLog.h"
+
+UTargetModifierPerNodeHierarchyRules::UTargetModifierPerNodeHierarchyRules()
+{
+	Root = CreateDefaultSubobject<USingleModifierPerNodeWithTargetSettings>(TEXT("Root"));
+}
+
+FName UTargetModifierPerNodeHierarchyRules::GetRootNode_Implementation() const
+{
+	return ensure(Root) ? Root->NodeName : NAME_None;
+}
+
+bool UTargetModifierPerNodeHierarchyRules::GetParentNode_Implementation(FName ChildGroup, FName& ParentGroup) const
+{
+	bool bFound = false;
+	ForEachGroup([ChildGroup, &ParentGroup, &bFound](USingleModifierPerNodeWithTargetSettings& Group, USingleModifierPerNodeWithTargetSettings* Parent)
+	{
+		if (Group.NodeName == ChildGroup)
+		{
+			ParentGroup = Parent->NodeName;
+			bFound = true;
+			return EBreakBehavior::Break;
+		}
+		return EBreakBehavior::Continue;
+	});
+	return bFound;
+}
+
+TSet<FName> UTargetModifierPerNodeHierarchyRules::GetChildNodes_Implementation(FName ParentGroup) const
+{
+	if (const USingleModifierPerNodeWithTargetSettings* Group = FindNodeByName(ParentGroup))
+	{
+		TSet<FName> Result;
+		Algo::TransformIf(Group->ChildElements, Result, [](const TObjectPtr<USingleModifierPerNodeWithTargetSettings>& Child){ return Child != nullptr; }, [](const TObjectPtr<USingleModifierPerNodeWithTargetSettings>& Child){ return Child->NodeName; });
+		return Result;
+	}
+	return {};
+}
+
+UVCamModifier* UTargetModifierPerNodeHierarchyRules::GetModifierInNode_Implementation(UVCamComponent* Component, FName GroupName) const
+{
+	if (const USingleModifierPerNodeWithTargetSettings* Group = FindNodeByName(GroupName))
+	{
+		UVCamModifier* Modifier = Component->GetModifierByName(Group->TargetSettings.TargetModifierName);
+		return Modifier;
+	}
+
+	return {};
+}
+
+bool UTargetModifierPerNodeHierarchyRules::GetConnectionPointTargetForNode_Implementation(FName GroupName, UVCamComponent* Component, FVCamModifierConnectionBinding& Connections) const
+{
+	USingleModifierPerNodeWithTargetSettings* GroupNode = FindNodeByName(GroupName);
+	if (!GroupNode)
+	{
+		return false;
+	}
+
+	const FVCamConnectionTargetSettings& TargetSettings = GroupNode->TargetSettings;
+	UVCamModifier* Modifier = Component->GetModifierByName(TargetSettings.TargetModifierName);
+	if (!Modifier)
+	{
+		return false;
+	}
+	
+	Connections = { Modifier, TargetSettings.TargetConnectionPoint };
+	return true;
+}
+
+TSet<FName> UTargetModifierPerNodeHierarchyRules::GetNodesContainingModifier_Implementation(UVCamModifier* Modifier) const
+{
+	TSet<FName> Groups;
+	Algo::Transform(GetNodeForModifier(Modifier), Groups, [](USingleModifierPerNodeWithTargetSettings* Node){ return Node->NodeName; });
+	return Groups;
+}
+
+TSet<USingleModifierPerNodeWithTargetSettings*> UTargetModifierPerNodeHierarchyRules::GetNodeForModifier(UVCamModifier* Modifier) const
+{
+	if (!Modifier)
+	{
+		return {};
+	}
+
+	TSet<USingleModifierPerNodeWithTargetSettings*> FoundGroups;
+	ForEachGroup([&FoundGroups, SearchName = Modifier->GetStackEntryName()](USingleModifierPerNodeWithTargetSettings& CurrentGroup, USingleModifierPerNodeWithTargetSettings* Parent)
+	{
+		if (CurrentGroup.TargetSettings.TargetModifierName == SearchName)
+		{
+			FoundGroups.Add(&CurrentGroup);
+		}
+		return EBreakBehavior::Continue;
+	});
+	return FoundGroups;
+}
+
+USingleModifierPerNodeWithTargetSettings* UTargetModifierPerNodeHierarchyRules::FindNodeByName(FName GroupName) const
+{
+	USingleModifierPerNodeWithTargetSettings* FoundGroup = nullptr;
+	ForEachGroup([GroupName, &FoundGroup](USingleModifierPerNodeWithTargetSettings& Group, USingleModifierPerNodeWithTargetSettings* Parent)
+	{
+		if (Group.NodeName == GroupName)
+		{
+			FoundGroup = &Group;
+			return EBreakBehavior::Break;
+		}
+		return EBreakBehavior::Continue;
+	});
+	UE_CLOG(FoundGroup == nullptr, LogVCamExtensions, Warning, TEXT("GroupName %s not found (%s)"), *GroupName.ToString(), *GetPathName());
+	return FoundGroup;
+}
+
+void UTargetModifierPerNodeHierarchyRules::ForEachGroup(TFunctionRef<EBreakBehavior(USingleModifierPerNodeWithTargetSettings& CurrentGroup, USingleModifierPerNodeWithTargetSettings* Parent)> Callback) const
+{
+	if (!ensure(Root))
+	{
+		return;
+	}
+
+	using namespace UE::VCamExtensions;
+	HierarchyUtils::ForEachGroup(
+		*Root,
+		[Callback](USingleModifierPerNodeWithTargetSettings& CurrentGroup, USingleModifierPerNodeWithTargetSettings* Parent)
+		{
+			return Callback(CurrentGroup, Parent) == EBreakBehavior::Continue ? HierarchyUtils::EBreakBehavior::Continue : HierarchyUtils::EBreakBehavior::Break;
+		},
+		[](USingleModifierPerNodeWithTargetSettings& CurrentGroup){ return CurrentGroup.ChildElements; }
+		);
+}
diff --git a/Engine/Plugins/Experimental/VirtualCamera/Source/VCamExtensions/Public/Hierarchies/BaseModifierGroup.h b/Engine/Plugins/Experimental/VirtualCamera/Source/VCamExtensions/Public/Hierarchies/BaseModifierGroup.h
--- a/Engine/Plugins/Experimental/VirtualCamera/Source/VCamExtensions/Public/Hierarchies/BaseModifierGroup.h
+++ b/Engine/Plugins/Experimental/VirtualCamera/Source/VCamExtensions/Public/Hierarchies/BaseModifierGroup.h
@@ -13,15 +13,9 @@ class VCAMEXTENSIONS_API UBaseModifierGroup : public UObject
 public:
 
 	UPROPERTY(EditAnywhere, Category = "Hierarchy")
-	FName GroupName;
+	FName NodeName;
 
 	//~ Begin UObject Interface
 	virtual void PostInitProperties() override;
 	//~ End UObject Interface
-
-#if WITH_EDITOR
-	//~ Begin UObject Interface
-	virtual void PostEditChangeProperty(FPropertyChangedEvent& PropertyChangedEvent) override;
-	//~ End UObject Interface
-#endif
 };
diff --git a/Engine/Plugins/Experimental/VirtualCamera/Source/VCamExtensions/Public/Hierarchies/ClassBasedModifierHierarchyRules.h b/Engine/Plugins/Experimental/VirtualCamera/Source/VCamExtensions/Public/Hierarchies/ClassBasedModifierHierarchyRules.h
deleted file mode 100644
diff --git a/Engine/Plugins/Experimental/VirtualCamera/Source/VCamExtensions/Public/Hierarchies/ModifierHierarchyAsset.h b/Engine/Plugins/Experimental/VirtualCamera/Source/VCamExtensions/Public/Hierarchies/ModifierHierarchyAsset.h
--- a/Engine/Plugins/Experimental/VirtualCamera/Source/VCamExtensions/Public/Hierarchies/ModifierHierarchyAsset.h
+++ b/Engine/Plugins/Experimental/VirtualCamera/Source/VCamExtensions/Public/Hierarchies/ModifierHierarchyAsset.h
@@ -9,6 +9,8 @@ class UModifierHierarchyRules;
 class UVCamComponent;
 class UVCamModifier;
 
+struct FVCamModifierConnectionBinding;
+
 /**
  * An asset intended to be referenced by Slate widgets.
  *
@@ -30,21 +32,29 @@ public:
 	TObjectPtr<UModifierHierarchyRules> Rules;
 
 	/** Gets the root of the tree. */
-	UFUNCTION(BlueprintPure, Category = "Virtual Camera")
-	FName GetRootGroup() const;
-
-	/**
-	 * Gets the group the modifier belongs to.
-	 * @return True if the modifier belongs to any group
-	 */
-	UFUNCTION(BlueprintPure, Category = "Virtual Camera")
-	bool GetGroupOfModifier(UVCamModifier* Modifier, FName& Group) const;
+	UFUNCTION(BlueprintPure, Category = "Virtual Camera|Hierarchies")
+	FName GetRootNode() const;
 
+	/** Gets the parent of this given group. Fails if called on the root node. */
+	UFUNCTION(BlueprintPure, Category = "Virtual Camera|Hierarchies")
+	bool GetParentNode(FName ChildNode, FName& ParentNode) const;
+	
 	/** Gets the child groups of the given group. */
-	UFUNCTION(BlueprintPure, BlueprintNativeEvent, Category = "Virtual Camera")
-	TSet<FName> GetChildGroups(FName ParentGroup) const;
+	UFUNCTION(BlueprintPure, Category = "Virtual Camera|Hierarchies")
+	TSet<FName> GetChildNodes(FName ParentGroup) const;
 
 	/** Gets all the modifiers on the component that belong in the given group. */
-	UFUNCTION(BlueprintPure, Category = "Virtual Camera")
-	TSet<UVCamModifier*> GetModifiersInGroup(UVCamComponent* Component, FName GroupName) const;
+	UFUNCTION(BlueprintPure, Category = "Virtual Camera|Hierarchies")
+	UVCamModifier* GetModifierInNode(UVCamComponent* Component, FName GroupName) const;
+	
+	/**
+	 * Gets the connection point the modifier is configured to be bound to.
+	 * This function is optional to implement; it is valid for it to always return false.
+	 */
+	UFUNCTION(BlueprintPure, Category = "Virtual Camera|Hierarchies")
+	bool GetConnectionPointTargetForNode(FName GroupName, UVCamComponent* Component, FVCamModifierConnectionBinding& Connection) const;
+	
+	/** Utility function to get all groups which contain this modifier. */
+	UFUNCTION(BlueprintPure, Category = "Virtual Camera|Hierarchies")
+	TSet<FName> GetGroupsContainingModifier(UVCamModifier* Modifier) const;
 };
\ No newline at end of file
diff --git a/Engine/Plugins/Experimental/VirtualCamera/Source/VCamExtensions/Public/Hierarchies/ModifierHierarchyRules.h b/Engine/Plugins/Experimental/VirtualCamera/Source/VCamExtensions/Public/Hierarchies/ModifierHierarchyRules.h
--- a/Engine/Plugins/Experimental/VirtualCamera/Source/VCamExtensions/Public/Hierarchies/ModifierHierarchyRules.h
+++ b/Engine/Plugins/Experimental/VirtualCamera/Source/VCamExtensions/Public/Hierarchies/ModifierHierarchyRules.h
@@ -2,16 +2,28 @@
 
 #pragma once
 
+#include "UI/VCamConnectionStructs.h"
 #include "UObject/Object.h"
 #include "ModifierHierarchyRules.generated.h"
 
 class UVCamComponent;
 class UVCamModifier;
 
+USTRUCT(BlueprintType)
+struct FVCamModifierConnectionBinding
+{
+	GENERATED_BODY()
+
+	UPROPERTY(BlueprintReadWrite, Category = "Virtual Camera")
+	TObjectPtr<UVCamModifier> Modifier;
+
+	UPROPERTY(BlueprintReadWrite, Category = "Virtual Camera")
+	FName ConnectionPoint;
+};
+
 /**
  * Defines a tree hierarchy. Each node is called a group.
- * A group consists of modifiers and (sub) groups.
- * By design, modifiers are restricted to only belong to at most one group.
+ * A group consists of connections and (sub) groups. A connection is a modifier and a corresponding connection point.
  * 
  * An example use case is if you want to have a button menu which should procedurally generate sub-button menus depending
  * on the modifiers in the component. 
@@ -26,33 +38,29 @@ public:
 
 	/** Gets the root of the tree. */
 	UFUNCTION(BlueprintPure, BlueprintNativeEvent, Category = "Virtual Camera|Hierarchies")
-	FName GetRootGroup() const;
+	FName GetRootNode() const;
 
 	/** Gets the parent of this given group. Fails if called on the root node. */
 	UFUNCTION(BlueprintPure, BlueprintNativeEvent, Category = "Virtual Camera|Hierarchies")
-	bool GetParentGroup(FName ChildGroup, FName& ParentGroup) const;
-
+	bool GetParentNode(FName ChildNode, FName& ParentNode) const;
 	
-	/**
-	 * Gets the group the modifier belongs to.
-	 * @return True if the modifier belongs to any group
-	 */
+	/** Gets the child groups of the given group. */
 	UFUNCTION(BlueprintPure, BlueprintNativeEvent, Category = "Virtual Camera|Hierarchies")
-	bool GetGroupOfModifier(UVCamModifier* Modifier, FName& Group) const;
+	TSet<FName> GetChildNodes(FName Node) const;
+
+	/** Gets all the modifiers on the component that belong in the given group. */
+	UFUNCTION(BlueprintPure, BlueprintNativeEvent, Category = "Virtual Camera|Hierarchies")
+	UVCamModifier* GetModifierInNode(UVCamComponent* Component, FName NodeName) const;
 	
 	/**
-	 * Gets the connection point the modifier is configured to be bound to.
+	 * Gets the connection point the modifier is configured to be bound to, if any.
 	 * This function is optional to implement; it is valid for it to always return false.
+	 * @return Whether there is a connection point configured for this hierarchy node.
 	 */
 	UFUNCTION(BlueprintPure, BlueprintNativeEvent, Category = "Virtual Camera|Hierarchies")
-	bool GetConnectionPointForModifier(UVCamModifier* Modifier, FName& ConnectionPoint) const;
-
+	bool GetConnectionPointTargetForNode(FName GroupName, UVCamComponent* Component, FVCamModifierConnectionBinding& Connection) const;
 	
-	/** Gets the child groups of the given group. */
-	UFUNCTION(BlueprintPure, BlueprintNativeEvent, Category = "Virtual Camera|Hierarchies")
-	TSet<FName> GetChildGroups(FName ParentGroup) const;
-
-	/** Gets all the modifiers on the component that belong in the given group. */
+	/** Utility function to get all groups which contain this modifier. */
 	UFUNCTION(BlueprintPure, BlueprintNativeEvent, Category = "Virtual Camera|Hierarchies")
-	TSet<UVCamModifier*> GetModifiersInGroup(UVCamComponent* Component, FName GroupName) const;
+	TSet<FName> GetNodesContainingModifier(UVCamModifier* Modifier) const;
 };
diff --git a/Engine/Plugins/Experimental/VirtualCamera/Source/VCamExtensions/Public/Hierarchies/TargetModifierPerGroupHierarchyRules.h b/Engine/Plugins/Experimental/VirtualCamera/Source/VCamExtensions/Public/Hierarchies/TargetModifierPerGroupHierarchyRules.h
deleted file mode 100644
diff --git a/Engine/Plugins/Experimental/VirtualCamera/Source/VCamExtensions/Public/Hierarchies/TargetModifierPerNodeHierarchyRules.h b/Engine/Plugins/Experimental/VirtualCamera/Source/VCamExtensions/Public/Hierarchies/TargetModifierPerNodeHierarchyRules.h
new file mode 100644
--- /dev/null
+++ b/Engine/Plugins/Experimental/VirtualCamera/Source/VCamExtensions/Public/Hierarchies/TargetModifierPerNodeHierarchyRules.h
@@ -0,0 +1,58 @@
+// Copyright Epic Games, Inc. All Rights Reserved.
+
+#pragma once
+
+#include "BaseModifierGroup.h"
+#include "ModifierHierarchyRules.h"
+#include "Templates/Function.h"
+#include "UI/VCamConnectionStructs.h"
+#include "TargetModifierPerNodeHierarchyRules.generated.h"
+
+UCLASS(EditInlineNew)
+class VCAMEXTENSIONS_API USingleModifierPerNodeWithTargetSettings : public UBaseModifierGroup
+{
+	GENERATED_BODY()
+public:
+
+	UPROPERTY(EditAnywhere, Category = "Hierarchy")
+	FVCamConnectionTargetSettings TargetSettings;
+
+	UPROPERTY(EditAnywhere, Instanced, Category = "Hierarchy")
+	TArray<TObjectPtr<USingleModifierPerNodeWithTargetSettings>> ChildElements;
+};
+
+/**
+ * 
+ */
+UCLASS()
+class VCAMEXTENSIONS_API UTargetModifierPerNodeHierarchyRules : public UModifierHierarchyRules
+{
+	GENERATED_BODY()
+public:
+	
+	UTargetModifierPerNodeHierarchyRules();
+
+	//~ Begin UModifierHierarchyRules Interface
+	virtual FName GetRootNode_Implementation() const override;
+	virtual bool GetParentNode_Implementation(FName ChildGroup, FName& ParentGroup) const override;
+	virtual TSet<FName> GetChildNodes_Implementation(FName ParentGroup) const override;
+	virtual UVCamModifier* GetModifierInNode_Implementation(UVCamComponent* Component, FName GroupName) const override;
+	virtual bool GetConnectionPointTargetForNode_Implementation(FName GroupName, UVCamComponent* Component, FVCamModifierConnectionBinding& Connection) const override;
+	virtual TSet<FName> GetNodesContainingModifier_Implementation(UVCamModifier* Modifier) const override;
+	//~ End UModifierHierarchyRules Interface
+
+private:
+
+	UPROPERTY(EditAnywhere, Instanced, Category = "Hierarchy")
+	TObjectPtr<USingleModifierPerNodeWithTargetSettings> Root;
+
+	TSet<USingleModifierPerNodeWithTargetSettings*> GetNodeForModifier(UVCamModifier* Modifier) const;
+	USingleModifierPerNodeWithTargetSettings* FindNodeByName(FName GroupName) const;
+	
+	enum class EBreakBehavior
+	{
+		Continue,
+		Break
+	};
+	void ForEachGroup(TFunctionRef<EBreakBehavior(USingleModifierPerNodeWithTargetSettings& CurrentGroup, USingleModifierPerNodeWithTargetSettings* Parent)> Callback) const;
+};
