Move code to VirtualCameraCore so VirtualCamera content can be configured to only load on Windows and Linux.

diff --git a/Engine/Plugins/Experimental/VirtualCamera/Config/BaseEngine.ini b/Engine/Plugins/Experimental/VirtualCamera/Config/BaseEngine.ini
--- a/Engine/Plugins/Experimental/VirtualCamera/Config/BaseEngine.ini
+++ b/Engine/Plugins/Experimental/VirtualCamera/Config/BaseEngine.ini
@@ -1,10 +1,4 @@
 [CoreRedirects]
-+ClassRedirects=(OldName="/Script/VCamCore.VCamCoreSubsystem",NewName="/Script/VCamInput.VCamInputSubsystem")
-+ClassRedirects=(OldName="/Script/VirtualCamera.RadialSlider",NewName="/Script/AdvancedWidgets.RadialSlider")
-+FunctionRedirects=(OldName="VCamComponent.ClearModifiers",NewName="VCamComponent.RemoveAllModifiers")
-+FunctionRedirects=(OldName="VCamComponent.FindModifierByName",NewName="VCamComponent.GetModifierByName")
-+FunctionRedirects=(OldName="VCamComponent.FindModifiersByClass",NewName="VCamComponent.GetModifiersByClass")
-+FunctionRedirects=(OldName="VCamComponent.FindModifierByInterface",NewName="VCamComponent.GetModifiersByInterface")
 +PackageRedirects=(OldName="/VirtualCamera/Editor", NewName="/VirtualCamera/Deprecated/Editor", MatchSubstring=true)
 +PackageRedirects=(OldName="/VirtualCamera/OperatorPanel", NewName="/VirtualCamera/Deprecated/OperatorPanel", MatchSubstring=true)
 +PackageRedirects=(OldName="/VirtualCamera/TestMaps", NewName="/VirtualCamera/Deprecated/TestMaps", MatchSubstring=true)
diff --git a/Engine/Plugins/Experimental/VirtualCamera/VirtualCamera.uplugin b/Engine/Plugins/Experimental/VirtualCamera/VirtualCamera.uplugin
--- a/Engine/Plugins/Experimental/VirtualCamera/VirtualCamera.uplugin
+++ b/Engine/Plugins/Experimental/VirtualCamera/VirtualCamera.uplugin
@@ -3,7 +3,7 @@
 	"Version": 1,
 	"VersionName": "1.0",
 	"FriendlyName": "VirtualCamera",
-	"Description": "Actors, components, and utilities for controlling and viewing cameras via physical devices",
+	"Description": "Content for VirtualCameraCore which adds actors, components, and utilities for controlling and viewing cameras via physical devices.",
 	"Category": "Virtual Production",
 	"CreatedBy": "Epic Games, Inc.",
 	"CreatedByURL": "https://epicgames.com",
@@ -22,121 +22,9 @@
 		"Win64",
 		"Linux"
 	],
-	"Modules": [
-		{
-			"Name": "PixelStreamingVCam",
-			"Type": "Editor",
-			"LoadingPhase": "PostDefault",
-			"PlatformAllowList": [ "Win64", "Linux" ],
-			"TargetDenyList": [ "Server" ]
-		},
-		{
-			"Name": "VCamBlueprintNodes",
-			"Type": "UncookedOnly",
-			"LoadingPhase": "Default"
-		},
-		{
-			"Name": "VCamCore",
-			"Type": "Runtime",
-			"LoadingPhase": "Default",
-			"PlatformDenyList": [
-				"TVOS"
-			]
-		},
-		{
-			"Name": "VCamCoreEditor",
-			"Type": "Editor",
-			"LoadingPhase": "Default",
-			"PlatformDenyList": [
-				"TVOS"
-			]
-		},
-		{
-			"Name": "VCamExtensions",
-			"Type": "Runtime",
-			"LoadingPhase": "Default",
-			"PlatformDenyList": [
-				"TVOS"
-			]
-		},
-		{
-			"Name": "VCamExtensionsEditor",
-			"Type": "Editor",
-			"LoadingPhase": "Default",
-			"PlatformDenyList": [
-				"TVOS"
-			]
-		},
-		{
-			"Name": "VCamInput",
-			"Type": "Runtime",
-			"LoadingPhase": "PreDefault",
-			"PlatformDenyList": [
-				"TVOS"
-			]
-		},
-		{
-			"Name": "VirtualCamera",
-			"Type": "Runtime",
-			"LoadingPhase": "PostDefault",
-			"PlatformAllowList": [
-				"Win64",
-				"Linux"
-			]
-		},
-		{
-			"Name": "VirtualCameraEditor",
-			"Type": "Editor",
-			"LoadingPhase": "PostEngineInit"
-		}
-	],
 	"Plugins": [
 		{
-			"Name": "MediaIOFramework",
-			"Enabled": true
-		},
-		{
-			"Name": "MediaFrameworkUtilities",
-			"Enabled": true
-		},
-		{
-			"Name": "RemoteSession",
-			"Enabled": true
-		},
-		{
-			"Name": "LevelSequenceEditor",
-			"Enabled": true
-		},
-		{
-			"Name": "Takes",
-			"Enabled": true
-		},
-		{
-			"Name": "VirtualProductionUtilities",
-			"Enabled": true
-		},
-		{
-			"Name": "VPRoles",
-			"Enabled": true
-		},
-		{
-			"Name": "LiveLink",
-			"Enabled": true
-		},
-		{
-			"Name": "EnhancedInput",
-			"Enabled": true
-		},
-		{
-			"Name": "CommonUI",
-			"Enabled": true
-		},
-		{
-			"Name": "PixelStreaming",
-			"Enabled": true
-		},
-		{
-			"Name" : "PixelCapture",
+			"Name" : "VirtualCameraCore",
 			"Enabled": true
 		}
 	]
