Initial checking of SplineModifier using CineCameraRigRail

diff --git a/Engine/Plugins/Experimental/VirtualCamera/VirtualCamera.uplugin b/Engine/Plugins/Experimental/VirtualCamera/VirtualCamera.uplugin
--- a/Engine/Plugins/Experimental/VirtualCamera/VirtualCamera.uplugin
+++ b/Engine/Plugins/Experimental/VirtualCamera/VirtualCamera.uplugin
@@ -26,6 +26,10 @@
 		{
 			"Name" : "VirtualCameraCore",
 			"Enabled": true
+		},
+		{
+			"Name" : "CineCameraRigs",
+			"Enabled": true
 		}
 	]
 }
