Revert change to VirtualCamera plugin to have split Mac / Windows & Linux

#rnx

diff --git a/Engine/Plugins/Experimental/VirtualCamera/VirtualCamera.uplugin b/Engine/Plugins/Experimental/VirtualCamera/VirtualCamera.uplugin
--- a/Engine/Plugins/Experimental/VirtualCamera/VirtualCamera.uplugin
+++ b/Engine/Plugins/Experimental/VirtualCamera/VirtualCamera.uplugin
@@ -14,6 +14,14 @@
 	"IsBetaVersion": true,
 	"Installed": false,
 	"EnabledByDefault": false,
+ 	"SupportedTargetPlatforms": [
+ 		"Win64",
+ 		"Linux"
+ 	],
+ 	"PlatformAllowList": [
+ 		"Win64",
+ 		"Linux"
+ 	],
 	"Plugins": [
 		{
 			"Name" : "VirtualCameraCore",
@@ -25,13 +33,7 @@
 		},
 		{
 			"Name" : "VCamPixelStreamingPreset",
-			"Enabled": true,
-			"PlatformAllowList" : [ "Win64", "Linux" ]
+			"Enabled": true
 		},
-		{
-			"Name" : "VCamMacPreset",
-			"Enabled": true,
-			"PlatformAllowList" : [ "Mac" ]
-		}
 	]
 }
