Move VirtualCamera and VirtualCameraEditor to VirtualCamera plugin

diff --git a/Engine/Plugins/Experimental/VCamMacPreset/Source/VCamMacPresetEditor/Private/VCamMacPresetEditorModule.cpp b/Engine/Plugins/Experimental/VCamMacPreset/Source/VCamMacPresetEditor/Private/VCamMacPresetEditorModule.cpp
--- a/Engine/Plugins/Experimental/VCamMacPreset/Source/VCamMacPresetEditor/Private/VCamMacPresetEditorModule.cpp
+++ b/Engine/Plugins/Experimental/VCamMacPreset/Source/VCamMacPresetEditor/Private/VCamMacPresetEditorModule.cpp
@@ -9,7 +9,6 @@
 #include "LevelEditor.h"
 #include "LevelEditorOutlinerSettings.h"
 #include "Modules/ModuleManager.h"
-#include "VirtualCameraActor.h"
 
 #define LOCTEXT_NAMESPACE "FVirtualCameraMacPresetEditorModule"
 
diff --git a/Engine/Plugins/Experimental/VCamMacPreset/Source/VCamMacPresetEditor/VCamMacPresetEditor.Build.cs b/Engine/Plugins/Experimental/VCamMacPreset/Source/VCamMacPresetEditor/VCamMacPresetEditor.Build.cs
--- a/Engine/Plugins/Experimental/VCamMacPreset/Source/VCamMacPresetEditor/VCamMacPresetEditor.Build.cs
+++ b/Engine/Plugins/Experimental/VCamMacPreset/Source/VCamMacPresetEditor/VCamMacPresetEditor.Build.cs
@@ -21,8 +21,6 @@ public class VCamMacPresetEditor : ModuleRules
 				"PlacementMode",
 				"VPUtilities",
 				"VPUtilitiesEditor",
-				"VCamCore",
-				"VirtualCamera",
 			}
 		);
 			
diff --git a/Engine/Plugins/Experimental/VirtualCamera/Source/VirtualCamera/Private/ConcertVirtualCamera.cpp b/Engine/Plugins/Experimental/VirtualCamera/Source/VirtualCamera/Private/ConcertVirtualCamera.cpp
new file mode 100644
--- /dev/null
+++ b/Engine/Plugins/Experimental/VirtualCamera/Source/VirtualCamera/Private/ConcertVirtualCamera.cpp
@@ -0,0 +1,290 @@
+// Copyright Epic Games, Inc. All Rights Reserved.
+
+#include "ConcertVirtualCamera.h"
+
+#if VIRTUALCAMERA_WITH_CONCERT
+#include "CineCameraComponent.h"
+#include "GameFramework/Actor.h"
+#include "IConcertModule.h"
+#include "IConcertClient.h"
+#include "IConcertSession.h"
+#include "IConcertSyncClient.h"
+#include "IMultiUserClientModule.h"
+#include "VirtualCameraCineCameraComponent.h"
+#endif
+
+
+/**
+ *
+ */
+FConcertVirtualCameraCameraFocusData::FConcertVirtualCameraCameraFocusData()
+	: ManualFocusDistance(100000.f)
+	, FocusSmoothingInterpSpeed(8.f)
+	, bSmoothFocusChanges(false)
+{}
+
+
+FConcertVirtualCameraCameraFocusData::FConcertVirtualCameraCameraFocusData(const UCineCameraComponent* CineCamera)
+	: ManualFocusDistance(CineCamera->CurrentFocusDistance)
+	, FocusSmoothingInterpSpeed(CineCamera->FocusSettings.FocusSmoothingInterpSpeed)
+	, bSmoothFocusChanges(CineCamera->FocusSettings.bSmoothFocusChanges)
+{}
+
+
+FCameraFocusSettings FConcertVirtualCameraCameraFocusData::ToCameraFocusSettings() const
+{
+	FCameraFocusSettings Result;
+	Result.FocusMethod = ECameraFocusMethod::Manual;
+	Result.ManualFocusDistance = ManualFocusDistance;
+	Result.bSmoothFocusChanges = bSmoothFocusChanges;
+	Result.FocusSmoothingInterpSpeed = FocusSmoothingInterpSpeed;
+	Result.FocusOffset = 0.f;
+	return Result;
+}
+
+
+/**
+ *
+ */
+FConcertVirtualCameraCameraData::FConcertVirtualCameraCameraData()
+	: CameraActorLocation(FVector::ZeroVector)
+	, CameraActorRotation(FRotator::ZeroRotator)
+	, CameraComponentLocation(FVector::ZeroVector)
+	, CameraComponentRotation(FRotator::ZeroRotator)
+	, CurrentAperture(0.f)
+	, CurrentFocalLength(0.f)
+{}
+
+
+FConcertVirtualCameraCameraData::FConcertVirtualCameraCameraData(const AActor* InOwner, const UCineCameraComponent* InCineCameraComponent)
+{
+	CameraActorLocation = InOwner->GetActorLocation();
+	CameraActorRotation = InOwner->GetActorRotation();
+	CameraComponentLocation = InCineCameraComponent->GetRelativeLocation();
+	CameraComponentRotation = InCineCameraComponent->GetRelativeRotation();
+
+	CurrentAperture = InCineCameraComponent->CurrentAperture;
+	CurrentFocalLength = InCineCameraComponent->CurrentFocalLength;
+	FocusSettings = FConcertVirtualCameraCameraFocusData(InCineCameraComponent);
+	LensSettings = InCineCameraComponent->LensSettings;
+
+	if (const UVirtualCameraCineCameraComponent* VCCineCamCmp = Cast<UVirtualCameraCineCameraComponent>(InCineCameraComponent))
+	{
+		FilmbackSettings = VCCineCamCmp->DesiredFilmbackSettings;
+	}
+	else
+	{
+		FilmbackSettings = InCineCameraComponent->Filmback;
+	}
+}
+
+
+void FConcertVirtualCameraCameraData::ApplyTo(AActor* InOwner, UCineCameraComponent* InCineCameraComponent, bool bUpdateCameraComponentTransform)
+{
+	InOwner->SetActorLocationAndRotation(CameraActorLocation, CameraActorRotation);
+	if (bUpdateCameraComponentTransform)
+	{
+		InCineCameraComponent->SetRelativeLocationAndRotation(CameraComponentLocation, CameraComponentRotation);
+	}
+	InCineCameraComponent->CurrentAperture = CurrentAperture;
+	InCineCameraComponent->CurrentFocalLength = CurrentFocalLength;
+	InCineCameraComponent->FocusSettings = FocusSettings.ToCameraFocusSettings();
+	InCineCameraComponent->LensSettings = LensSettings;
+	InCineCameraComponent->Filmback = FilmbackSettings;
+
+	if (UVirtualCameraCineCameraComponent* VCCineCamCmp = Cast<UVirtualCameraCineCameraComponent>(InCineCameraComponent))
+	{
+		VCCineCamCmp->DesiredFilmbackSettings = FilmbackSettings;
+	}
+}
+
+
+/**
+ *
+ */
+FConcertVirtualCameraCameraComponentEvent::FConcertVirtualCameraCameraComponentEvent()
+{}
+
+
+/**
+ *
+ */
+FConcertVirtualCameraControllerEvent::FConcertVirtualCameraControllerEvent()
+	: InputSource(ETrackerInputSource::ARKit)
+{}
+
+
+
+#if VIRTUALCAMERA_WITH_CONCERT
+
+/**
+ *
+ */
+FConcertVirtualCameraManager::FConcertVirtualCameraManager()
+	: bIsLatestControllerCameraEventDataValid(false)
+{
+	if (TSharedPtr<IConcertSyncClient> ConcertSyncClient = IMultiUserClientModule::Get().GetClient())
+	{
+		IConcertClientRef ConcertClient = ConcertSyncClient->GetConcertClient();
+
+		OnSessionStartupHandle = ConcertClient->OnSessionStartup().AddRaw(this, &FConcertVirtualCameraManager::RegisterConcertSyncHandlers);
+		OnSessionShutdownHandle = ConcertClient->OnSessionShutdown().AddRaw(this, &FConcertVirtualCameraManager::UnregisterConcertSyncHandlers);
+
+		TSharedPtr<IConcertClientSession> ConcertClientSession = ConcertClient->GetCurrentSession();
+		if (ConcertClientSession.IsValid())
+		{
+			RegisterConcertSyncHandlers(ConcertClientSession.ToSharedRef());
+		}
+	}
+}
+
+
+FConcertVirtualCameraManager::~FConcertVirtualCameraManager()
+{
+	if (IMultiUserClientModule::IsAvailable())
+	{
+		if (TSharedPtr<IConcertSyncClient> ConcertSyncClient = IMultiUserClientModule::Get().GetClient())
+		{
+			IConcertClientRef ConcertClient = ConcertSyncClient->GetConcertClient();
+
+			TSharedPtr<IConcertClientSession> ConcertClientSession = ConcertClient->GetCurrentSession();
+			if (ConcertClientSession.IsValid())
+			{
+				UnregisterConcertSyncHandlers(ConcertClientSession.ToSharedRef());
+			}
+
+			ConcertClient->OnSessionStartup().Remove(OnSessionStartupHandle);
+			OnSessionStartupHandle.Reset();
+
+			ConcertClient->OnSessionShutdown().Remove(OnSessionShutdownHandle);
+			OnSessionShutdownHandle.Reset();
+		}
+	}
+}
+
+
+bool FConcertVirtualCameraManager::GetLatestCameraComponentEvent(FName InTrackingName, FConcertVirtualCameraCameraData& OutCameraEvent) const
+{
+	for (const FConcertVirtualCameraCameraComponentEvent& CameraEvent : LatestCameraEventDatas)
+	{
+		if (CameraEvent.TrackingName == InTrackingName)
+		{
+			OutCameraEvent = CameraEvent.CameraData;
+			return true;
+		}
+	}
+
+	return false;
+}
+
+
+void FConcertVirtualCameraManager::SendCameraCompoentEvent(FName InTrackingName, const FConcertVirtualCameraCameraData& InCameraEvent)
+{
+	TSharedPtr<IConcertClientSession> Session = WeakSession.Pin();
+	if (Session.IsValid() && InTrackingName!= NAME_None)
+	{
+		FConcertVirtualCameraCameraComponentEvent CameraEvent;
+		CameraEvent.TrackingName = InTrackingName;
+		CameraEvent.CameraData = InCameraEvent;
+		TArray<FGuid> ClientIds = Session->GetSessionClientEndpointIds();
+		Session->SendCustomEvent(CameraEvent, ClientIds, EConcertMessageFlags::None);
+	}
+}
+
+
+bool FConcertVirtualCameraManager::GetLatestControllerCameraEvent(FConcertVirtualCameraControllerEvent& OutCameraEvent) const
+{
+	OutCameraEvent = LatestControllerCameraEventData;
+	return bIsLatestControllerCameraEventDataValid;
+}
+
+
+void FConcertVirtualCameraManager::SendControllerCameraEvent(const FConcertVirtualCameraControllerEvent& InCameraEvent)
+{
+	TSharedPtr<IConcertClientSession> Session = WeakSession.Pin();
+	if (Session.IsValid())
+	{
+		TArray<FGuid> ClientIds = Session->GetSessionClientEndpointIds();
+		Session->SendCustomEvent(InCameraEvent, ClientIds, EConcertMessageFlags::None);
+	}
+}
+
+
+void FConcertVirtualCameraManager::RegisterConcertSyncHandlers(TSharedRef<IConcertClientSession> InSession)
+{
+	// Hold onto the session so we can trigger events
+	WeakSession = InSession;
+
+	// Register our events
+	InSession->RegisterCustomEventHandler<FConcertVirtualCameraCameraComponentEvent>(this, &FConcertVirtualCameraManager::HandleCameraComponentEventData);
+	InSession->RegisterCustomEventHandler<FConcertVirtualCameraControllerEvent>(this, &FConcertVirtualCameraManager::HandleControllerCameraEventData);
+}
+
+
+void FConcertVirtualCameraManager::UnregisterConcertSyncHandlers(TSharedRef<IConcertClientSession> InSession)
+{
+	// Unregister our events and explicitly reset the session ptr
+	TSharedPtr<IConcertClientSession> Session = WeakSession.Pin();
+	if (Session.IsValid())
+	{
+		Session->UnregisterCustomEventHandler<FConcertVirtualCameraControllerEvent>(this);
+		Session->UnregisterCustomEventHandler<FConcertVirtualCameraCameraComponentEvent>(this);
+	}
+
+	WeakSession.Reset();
+}
+
+
+void FConcertVirtualCameraManager::HandleCameraComponentEventData(const FConcertSessionContext& InEventContext, const FConcertVirtualCameraCameraComponentEvent& InEvent)
+{
+	if (InEvent.TrackingName == NAME_None)
+	{
+		return;
+	}
+
+	for (FConcertVirtualCameraCameraComponentEvent& CameraEvent : LatestCameraEventDatas)
+	{
+		if (CameraEvent.TrackingName == InEvent.TrackingName)
+		{
+			CameraEvent.CameraData = InEvent.CameraData;
+			return;
+		}
+	}
+
+	LatestCameraEventDatas.Add(InEvent);
+}
+
+
+void FConcertVirtualCameraManager::HandleControllerCameraEventData(const FConcertSessionContext& InEventContext, const FConcertVirtualCameraControllerEvent& InEvent)
+{
+	LatestControllerCameraEventData = InEvent;
+	bIsLatestControllerCameraEventDataValid = true;
+}
+
+
+#else //#if VIRTUALCAMERA_WITH_CONCERT
+
+
+bool FConcertVirtualCameraManager::GetLatestCameraComponentEvent(FName TrackingName, FConcertVirtualCameraCameraData& OutCameraEvent) const
+{
+	return false;
+}
+
+
+void FConcertVirtualCameraManager::SendCameraCompoentEvent(FName InTrackingName, const FConcertVirtualCameraCameraData& InCameraEvent)
+{
+}
+
+
+bool FConcertVirtualCameraManager::GetLatestControllerCameraEvent(FConcertVirtualCameraControllerEvent& OutCameraEvent) const
+{
+	return false;
+}
+
+
+void FConcertVirtualCameraManager::SendControllerCameraEvent(const FConcertVirtualCameraControllerEvent& InCameraEvent)
+{
+}
+
+
+#endif //#if VIRTUALCAMERA_WITH_CONCERT
diff --git a/Engine/Plugins/Experimental/VirtualCamera/Source/VirtualCamera/Private/ConcertVirtualCamera.h b/Engine/Plugins/Experimental/VirtualCamera/Source/VirtualCamera/Private/ConcertVirtualCamera.h
new file mode 100644
--- /dev/null
+++ b/Engine/Plugins/Experimental/VirtualCamera/Source/VirtualCamera/Private/ConcertVirtualCamera.h
@@ -0,0 +1,177 @@
+// Copyright Epic Games, Inc. All Rights Reserved.
+
+#pragma once
+
+#include "CoreMinimal.h"
+#include "CineCameraComponent.h"
+#include "VirtualCameraPlayerControllerBase.h"
+#include "UObject/Object.h"
+#include "UObject/ObjectMacros.h"
+#include "ConcertVirtualCamera.generated.h"
+
+struct FConcertSessionContext;
+class AActor;
+class IConcertClientSession;
+class UCineCameraComponent;
+
+/**
+ * Can't use FCameraFocusSettings since it use a reference to an actor
+ * The camera will always be in Manual, and will transfer the CurrentFocusDistance.
+ */
+USTRUCT()
+struct FConcertVirtualCameraCameraFocusData
+{
+public:
+	GENERATED_BODY()
+
+	FConcertVirtualCameraCameraFocusData();
+	FConcertVirtualCameraCameraFocusData(const UCineCameraComponent* InCineCameraComponent);
+
+	FCameraFocusSettings ToCameraFocusSettings() const;
+
+	UPROPERTY()
+	float ManualFocusDistance;
+	UPROPERTY()
+	float FocusSmoothingInterpSpeed;
+	UPROPERTY()
+	uint8 bSmoothFocusChanges : 1;
+};
+
+
+/**
+ *
+ */
+USTRUCT()
+struct FConcertVirtualCameraCameraData
+{
+public:
+	GENERATED_BODY()
+
+	FConcertVirtualCameraCameraData();
+	FConcertVirtualCameraCameraData(const AActor* InOwner, const UCineCameraComponent* InCineCameraComponent);
+
+	void ApplyTo(AActor* InOwner, UCineCameraComponent* InCineCameraComponent, bool bUpdateCameraComponentTransform);
+
+	/** Camera transform */
+	UPROPERTY()
+	FVector CameraActorLocation;
+	UPROPERTY()
+	FRotator CameraActorRotation;
+	UPROPERTY()
+	FVector CameraComponentLocation;
+	UPROPERTY()
+	FRotator CameraComponentRotation;
+
+	/** Camera settings */
+	UPROPERTY()
+	float CurrentAperture;
+	UPROPERTY()
+	float CurrentFocalLength; 
+	UPROPERTY()
+	FConcertVirtualCameraCameraFocusData FocusSettings;
+	UPROPERTY()
+	FCameraLensSettings LensSettings;
+	UPROPERTY()
+	FCameraFilmbackSettings FilmbackSettings;
+};
+
+
+/**
+ *
+ */
+USTRUCT()
+struct FConcertVirtualCameraCameraComponentEvent
+{
+public:
+	GENERATED_BODY()
+
+	FConcertVirtualCameraCameraComponentEvent();
+
+	/** Name of the tracking camera */
+	UPROPERTY()
+	FName TrackingName;
+
+	/** Camera data */
+	UPROPERTY()
+	FConcertVirtualCameraCameraData CameraData;
+};
+
+
+/**
+ *
+ */
+USTRUCT()
+struct FConcertVirtualCameraControllerEvent
+{
+public:
+	GENERATED_BODY()
+
+	FConcertVirtualCameraControllerEvent();
+
+	/** Controller settings */
+	UPROPERTY()
+	ETrackerInputSource InputSource;
+
+	/** Camera data */
+	UPROPERTY()
+	FConcertVirtualCameraCameraData CameraData;
+};
+
+
+
+#if VIRTUALCAMERA_WITH_CONCERT // Concert is only available in development mode
+
+/**
+ *
+ */
+class FConcertVirtualCameraManager
+{
+public:
+	FConcertVirtualCameraManager();
+	FConcertVirtualCameraManager(const FConcertVirtualCameraManager&) = delete;
+	FConcertVirtualCameraManager& operator=(const FConcertVirtualCameraManager&) = delete;
+	~FConcertVirtualCameraManager();
+
+public:
+	bool GetLatestCameraComponentEvent(FName TrackingName, FConcertVirtualCameraCameraData& OutCameraData) const;
+	void SendCameraCompoentEvent(FName TrackingName, const FConcertVirtualCameraCameraData& InCameraData);
+
+	bool GetLatestControllerCameraEvent(FConcertVirtualCameraControllerEvent& OutCameraEvent) const;
+	void SendControllerCameraEvent(const FConcertVirtualCameraControllerEvent& InCameraEvent);
+
+private:
+	void RegisterConcertSyncHandlers(TSharedRef<IConcertClientSession> InSession);
+	void UnregisterConcertSyncHandlers(TSharedRef<IConcertClientSession> InSession);
+
+	void HandleCameraComponentEventData(const FConcertSessionContext& InEventContext, const FConcertVirtualCameraCameraComponentEvent& InEvent);
+	void HandleControllerCameraEventData(const FConcertSessionContext& InEventContext, const FConcertVirtualCameraControllerEvent& InEvent);
+
+private:
+	/** Latest event data */
+	TArray<FConcertVirtualCameraCameraComponentEvent> LatestCameraEventDatas;
+	FConcertVirtualCameraControllerEvent LatestControllerCameraEventData;
+	bool bIsLatestControllerCameraEventDataValid;
+
+	/** Delegate handle for a the callback when a session starts up */
+	FDelegateHandle OnSessionStartupHandle;
+
+	/** Delegate handle for a the callback when a session shuts down */
+	FDelegateHandle OnSessionShutdownHandle;
+
+	/** Weak pointer to the client session with which to send events. May be null or stale. */
+	TWeakPtr<IConcertClientSession> WeakSession;
+};
+
+#else //VIRTUALCAMERA_WITH_CONCERT
+
+class FConcertVirtualCameraManager
+{
+public:
+	bool GetLatestCameraComponentEvent(FName TrackingName, FConcertVirtualCameraCameraData& OutCameraData) const;
+	void SendCameraCompoentEvent(FName TrackingName, const FConcertVirtualCameraCameraData& InCameraData);
+
+	bool GetLatestControllerCameraEvent(FConcertVirtualCameraControllerEvent& OutCameraEvent) const;
+	void SendControllerCameraEvent(const FConcertVirtualCameraControllerEvent& InCameraEvent);
+};
+
+#endif //VIRTUALCAMERA_WITH_CONCERT
diff --git a/Engine/Plugins/Experimental/VirtualCamera/Source/VirtualCamera/Private/VirtualCameraActor.cpp b/Engine/Plugins/Experimental/VirtualCamera/Source/VirtualCamera/Private/VirtualCameraActor.cpp
new file mode 100644
--- /dev/null
+++ b/Engine/Plugins/Experimental/VirtualCamera/Source/VirtualCamera/Private/VirtualCameraActor.cpp
@@ -0,0 +1,781 @@
+// Copyright Epic Games, Inc. All Rights Reserved.
+
+#include "VirtualCameraActor.h"
+
+#include "CineCameraComponent.h"
+#include "Components/SceneCaptureComponent2D.h"
+#include "Engine/Blueprint.h"
+#include "Engine/Engine.h"
+#include "Engine/GameEngine.h"
+#include "Engine/GameInstance.h"
+#include "Framework/Application/SlateApplication.h"
+#include "ILiveLinkClient.h"
+#include "ImageProviders/RemoteSessionMediaOutput.h"
+#include "GameFramework/PlayerController.h"
+#include "Kismet/GameplayStatics.h"
+#include "RemoteSession.h"
+#include "Roles/LiveLinkTransformRole.h"
+#include "Slate/SceneViewport.h"
+#include "VirtualCamera.h"
+#include "VirtualCameraMovement.h"
+#include "VirtualCameraSubsystem.h"
+#include "Widgets/VPFullScreenUserWidget.h"
+#include "VCamComponent.h"
+
+#if WITH_EDITOR
+#include "AssetRegistry/AssetData.h"
+#include "AssetRegistry/AssetRegistryModule.h"
+#include "Blueprint/UserWidget.h"
+#include "Editor.h"
+#include "EditorSupportDelegates.h"
+#include "Editor/EditorEngine.h"
+#include "IAssetViewport.h"
+#include "LevelEditor.h"
+#include "LevelEditorViewport.h"
+#include "SceneView.h"
+#include "SLevelViewport.h"
+#endif
+
+namespace
+{
+	static const FName AssetRegistryName(TEXT("AssetRegistry"));
+	static const FName LevelEditorName(TEXT("LevelEditor"));
+	static const FString SavedSettingsSlotName(TEXT("SavedVirtualCameraSettings"));
+	static const TCHAR* DefaultCameraUMG = TEXT("/VirtualCamera/V2/Widgets/VCam2UI.VCam2UI_C");
+	static const FName DefaultLiveLinkSubjectName(TEXT("CameraTransform"));
+	static const FVector2D DefaultViewportResolution(1536, 1152);
+	static const float MaxFocusTraceDistance = 1000000.0f;
+
+	void FindSceneViewport(TWeakPtr<SWindow>& OutInputWindow, TWeakPtr<FSceneViewport>& OutSceneViewport)
+	{
+#if WITH_EDITOR
+		if (GIsEditor)
+		{
+			for (const FWorldContext& Context : GEngine->GetWorldContexts())
+			{
+				if (Context.WorldType == EWorldType::Editor)
+				{
+					if (FModuleManager::Get().IsModuleLoaded(LevelEditorName))
+					{
+						FLevelEditorModule& LevelEditorModule = FModuleManager::GetModuleChecked<FLevelEditorModule>(LevelEditorName);
+						TSharedPtr<IAssetViewport> ActiveLevelViewport = LevelEditorModule.GetFirstActiveViewport();
+						if (ActiveLevelViewport.IsValid())
+						{
+							OutSceneViewport = ActiveLevelViewport->GetSharedActiveViewport();
+							OutInputWindow = FSlateApplication::Get().FindWidgetWindow(ActiveLevelViewport->AsWidget());
+						}
+					}
+				}
+				else if (Context.WorldType == EWorldType::PIE)
+				{
+					FSlatePlayInEditorInfo* SlatePlayInEditorSession = GEditor->SlatePlayInEditorMap.Find(Context.ContextHandle);
+					if (SlatePlayInEditorSession)
+					{
+						if (SlatePlayInEditorSession->DestinationSlateViewport.IsValid())
+						{
+							TSharedPtr<IAssetViewport> DestinationLevelViewport = SlatePlayInEditorSession->DestinationSlateViewport.Pin();
+							OutSceneViewport = DestinationLevelViewport->GetSharedActiveViewport();
+							OutInputWindow = FSlateApplication::Get().FindWidgetWindow(DestinationLevelViewport->AsWidget());
+						}
+						else if (SlatePlayInEditorSession->SlatePlayInEditorWindowViewport.IsValid())
+						{
+							OutSceneViewport = SlatePlayInEditorSession->SlatePlayInEditorWindowViewport;
+							OutInputWindow = SlatePlayInEditorSession->SlatePlayInEditorWindow;
+						}
+					}
+				}
+			}
+		}
+		else
+#endif
+		{
+			UGameEngine* GameEngine = Cast<UGameEngine>(GEngine);
+			OutSceneViewport = GameEngine->SceneViewport;
+			OutInputWindow = GameEngine->GameViewportWindow;
+		}
+	}
+
+	bool DeprojectScreenToWorld(const FVector2D& InScreenPosition, FVector& OutWorldPosition, FVector& OutWorldDirection)
+	{
+		bool bSuccess = false;
+
+		for (const FWorldContext& Context : GEngine->GetWorldContexts())
+		{
+			if (Context.WorldType == EWorldType::PIE || Context.WorldType == EWorldType::Game)
+			{
+				APlayerController* PC = Context.OwningGameInstance->GetFirstLocalPlayerController(Context.World());
+				if (PC)
+				{
+					bSuccess |= PC->DeprojectScreenPositionToWorld(InScreenPosition.X, InScreenPosition.Y, OutWorldPosition, OutWorldDirection);
+					break;
+				}
+			}
+#if WITH_EDITOR
+			else if (Context.WorldType == EWorldType::Editor)
+			{
+				FLevelEditorModule& LevelEditorModule = FModuleManager::GetModuleChecked<FLevelEditorModule>(LevelEditorName);
+				TSharedPtr<SLevelViewport> ActiveLevelViewport = LevelEditorModule.GetFirstActiveLevelViewport();
+				if (ActiveLevelViewport.IsValid() && ActiveLevelViewport->GetActiveViewport())
+				{
+					FViewport* ActiveViewport = ActiveLevelViewport->GetActiveViewport();
+					FLevelEditorViewportClient& LevelViewportClient = ActiveLevelViewport->GetLevelViewportClient();
+					FSceneViewFamilyContext ViewFamily(FSceneViewFamily::ConstructionValues(
+						ActiveViewport,
+						LevelViewportClient.GetScene(),
+						LevelViewportClient.EngineShowFlags)
+						.SetRealtimeUpdate(true));
+					FSceneView* View = LevelViewportClient.CalcSceneView(&ViewFamily);
+
+					const FIntPoint ViewportSize = ActiveViewport->GetSizeXY();
+					const FIntRect ViewRect = FIntRect(0, 0, ViewportSize.X, ViewportSize.Y);
+					const FMatrix InvViewProjectionMatrix = View->ViewMatrices.GetInvViewProjectionMatrix();
+					FSceneView::DeprojectScreenToWorld(InScreenPosition, ViewRect, InvViewProjectionMatrix, OutWorldPosition, OutWorldDirection);
+					bSuccess = true;
+				}
+			}
+#endif
+		}
+
+		if (!bSuccess)
+		{
+			OutWorldPosition = FVector::ZeroVector;
+			OutWorldDirection = FVector::ZeroVector;
+		}
+		return bSuccess;
+	}
+}
+
+struct FVirtualCameraViewportSettings
+{
+	FIntPoint Size;
+	FVector2D CameraPosition;
+	TWeakObjectPtr<AActor> ActorLock;
+	bool bRealTime;
+	bool bDrawAxes;
+	bool bDisableInput;
+	bool bAllowCinematicControl;
+};
+
+int32 AVirtualCameraActor::PresetIndex = 1;
+
+AVirtualCameraActor::AVirtualCameraActor(const FObjectInitializer& ObjectInitializer)
+	: Super(ObjectInitializer)
+	, LiveLinkSubject{ DefaultLiveLinkSubjectName, ULiveLinkTransformRole::StaticClass() }
+	, TargetDeviceResolution(DefaultViewportResolution)
+	, RemoteSessionPort(IRemoteSessionModule::kDefaultPort)
+	, ActorWorld(nullptr)
+	, PreviousViewTarget(nullptr)
+	, bAllowFocusVisualization(true)
+	, DesiredDistanceUnits(EUnit::Meters)
+	, bSaveSettingsOnStopStreaming(false)
+	, bIsStreaming(false)
+	, ViewportSettingsBackup(nullptr)
+{
+	PrimaryActorTick.bCanEverTick = true;
+	PrimaryActorTick.bStartWithTickEnabled = true;
+
+	StreamedCamera = GetCineCameraComponent();
+	
+	SceneCaptureComponent = CreateDefaultSubobject<USceneCaptureComponent2D>("Scene Capture");
+	SceneCaptureComponent->SetupAttachment(GetCineCameraComponent());
+	SceneCaptureComponent->bCaptureEveryFrame = 0;
+	SceneCaptureComponent->bCaptureOnMovement = 0;
+
+	MovementComponent = CreateDefaultSubobject<UVirtualCameraMovement>("Movement Component");
+	MediaOutput = CreateDefaultSubobject<URemoteSessionMediaOutput>("Media Output");
+	CameraScreenWidget = CreateDefaultSubobject<UVPFullScreenUserWidget>("Camera UMG");
+	CameraScreenWidget->SetDisplayTypes(EVPWidgetDisplayType::PostProcessWithBlendMaterial, EVPWidgetDisplayType::Viewport, EVPWidgetDisplayType::PostProcessWithBlendMaterial);
+	CameraScreenWidget->GetPostProcessDisplayTypeWithBlendMaterialSettings().bReceiveHardwareInput = true;
+
+	if (GEngine)
+	{
+		UVirtualCameraSubsystem* SubSystem = GEngine->GetEngineSubsystem<UVirtualCameraSubsystem>();
+		if (SubSystem && !SubSystem->GetVirtualCameraController())
+		{
+			SubSystem->SetVirtualCameraController(this);
+		}
+	}
+}
+
+AVirtualCameraActor::AVirtualCameraActor(FVTableHelper& Helper)
+	: Super(Helper)
+{
+}
+
+AVirtualCameraActor::~AVirtualCameraActor() = default;
+
+void AVirtualCameraActor::Destroyed()
+{
+	if (CameraScreenWidget && CameraScreenWidget->IsDisplayed())
+	{
+		CameraScreenWidget->Hide();
+	}
+}
+
+#if WITH_EDITOR
+bool AVirtualCameraActor::ShouldTickIfViewportsOnly() const
+{
+	return true;
+}
+#endif
+
+bool AVirtualCameraActor::IsStreaming_Implementation() const
+{
+	return bIsStreaming;
+}
+
+bool AVirtualCameraActor::ShouldSaveSettingsOnStopStreaming_Implementation() const
+{
+	return bSaveSettingsOnStopStreaming;
+}
+
+void AVirtualCameraActor::SetSaveSettingsOnStopStreaming_Implementation(bool bShouldSave)
+{
+	bSaveSettingsOnStopStreaming = bShouldSave;
+}
+
+FTransform AVirtualCameraActor::GetRelativeTransform_Implementation() const
+{
+	UCineCameraComponent* CineCamera = GetCineCameraComponent();
+	return (CineCamera ? CineCamera->GetRelativeTransform() : FTransform::Identity);
+}
+
+void AVirtualCameraActor::AddBlendableToCamera_Implementation(const TScriptInterface<IBlendableInterface>& InBlendableToAdd, float InWeight)
+{
+	UCineCameraComponent* CineCamera = GetCineCameraComponent();
+
+	if (CineCamera)
+	{
+		CineCamera->PostProcessSettings.AddBlendable(InBlendableToAdd, InWeight);
+	}
+}
+
+void AVirtualCameraActor::SetTrackedActorForFocus_Implementation(AActor* InActorToTrack, const FVector& InTrackingPointOffset)
+{
+	UCineCameraComponent* CineCamera = GetCineCameraComponent();
+	if (CineCamera)
+	{
+		CineCamera->FocusSettings.TrackingFocusSettings.ActorToTrack = InActorToTrack;
+		CineCamera->FocusSettings.TrackingFocusSettings.RelativeOffset = InTrackingPointOffset;
+	}
+}
+
+void AVirtualCameraActor::SetFocusVisualization_Implementation(bool bInShowFocusVisualization)
+{
+#if WITH_EDITOR
+	UCineCameraComponent* CineCamera = GetCineCameraComponent();
+	if (CineCamera)
+	{
+		if (CineCamera->FocusSettings.FocusMethod == ECameraFocusMethod::DoNotOverride)
+		{
+			UE_LOG(LogVirtualCamera, Warning, TEXT("Camera focus mode is currently set to none, cannot display focus plane!"));
+			return;
+		}
+		CineCamera->FocusSettings.bDrawDebugFocusPlane = bInShowFocusVisualization;
+	}
+#endif
+}
+
+void AVirtualCameraActor::SetBeforeSetVirtualCameraTransformDelegate_Implementation(const FPreSetVirtualCameraTransform& InDelegate)
+{
+	OnPreSetVirtualCameraTransform = InDelegate;
+}
+
+void AVirtualCameraActor::SetOnActorClickedDelegate_Implementation(const FOnActorClickedDelegate& InDelegate)
+{
+	OnActorClickedDelegate = InDelegate;
+}
+
+void AVirtualCameraActor::AddOnVirtualCameraUpdatedDelegate_Implementation(const FVirtualCameraTickDelegate& InDelegate)
+{
+	OnVirtualCameraUpdatedDelegates.Add(InDelegate);
+}
+
+void AVirtualCameraActor::RemoveOnVirtualCameraUpdatedDelegate_Implementation(const FVirtualCameraTickDelegate& InDelegate)
+{
+	OnVirtualCameraUpdatedDelegates.Remove(InDelegate);
+}
+
+void AVirtualCameraActor::Tick(float DeltaSeconds)
+{
+	Super::Tick(DeltaSeconds);
+
+	if (CameraScreenWidget && CameraUMGClass)
+	{
+		CameraScreenWidget->Tick(DeltaSeconds);
+	}
+
+	FMinimalViewInfo ViewInfo;
+	CalcCamera(DeltaSeconds, ViewInfo);
+
+	if (OnVirtualCameraUpdatedDelegates.IsBound())
+	{
+		FEditorScriptExecutionGuard ScriptGuard;
+		OnVirtualCameraUpdatedDelegates.Broadcast(DeltaSeconds);
+	}
+}
+
+void AVirtualCameraActor::BeginPlay()
+{
+	Super::BeginPlay();
+
+	UVirtualCameraSubsystem* SubSystem = GEngine->GetEngineSubsystem<UVirtualCameraSubsystem>();
+	if (SubSystem && !SubSystem->GetVirtualCameraController())
+	{
+		SubSystem->SetVirtualCameraController(this);
+	}
+
+	StartStreaming();
+}
+
+void AVirtualCameraActor::EndPlay(const EEndPlayReason::Type EndPlayReason)
+{
+	Super::EndPlay(EndPlayReason);
+	StopStreaming();
+}
+
+bool AVirtualCameraActor::StartStreaming()
+{
+	UVCamComponent* VCamComponent = Cast<UVCamComponent>(GetComponentByClass(UVCamComponent::StaticClass()));
+	if (!VCamComponent || !VCamComponent->IsEnabled())
+	{
+		return false;
+	}
+
+	ActorWorld = GetWorld();
+	if (!ActorWorld)
+	{
+		return false;
+	}
+
+	if (bSaveSettingsOnStopStreaming)
+	{
+		LoadSettings();
+	}
+
+#if WITH_EDITOR
+	if (ActorWorld->WorldType == EWorldType::Editor)
+	{
+		ViewportSettingsBackup = MakeUnique<FVirtualCameraViewportSettings>();
+
+		FLevelEditorModule& LevelEditorModule = FModuleManager::GetModuleChecked<FLevelEditorModule>(LevelEditorName);
+		TSharedPtr<SLevelViewport> ActiveLevelViewport = LevelEditorModule.GetFirstActiveLevelViewport();
+		if (ActiveLevelViewport.IsValid())
+		{
+			ActiveLevelViewport->GetSharedActiveViewport()->SetFixedViewportSize(TargetDeviceResolution.X, TargetDeviceResolution.Y);
+
+			FLevelEditorViewportClient& LevelViewportClient = ActiveLevelViewport->GetLevelViewportClient();
+			ViewportSettingsBackup->ActorLock = LevelViewportClient.GetActiveActorLock();
+			LevelViewportClient.SetActorLock(this);
+
+			ViewportSettingsBackup->bDrawAxes = LevelViewportClient.bDrawAxes;
+			ViewportSettingsBackup->bDisableInput = LevelViewportClient.bDisableInput;
+			ViewportSettingsBackup->bAllowCinematicControl = LevelViewportClient.AllowsCinematicControl();
+
+			LevelViewportClient.SetRealtime(true);
+			LevelViewportClient.bDrawAxes = false;
+			LevelViewportClient.bDisableInput = true;
+			LevelViewportClient.SetAllowCinematicControl(false);
+
+			// add event listeners to stop streaming when necessary
+			LevelEditorModule.OnMapChanged().AddUObject(this, &AVirtualCameraActor::OnMapChanged);
+			GEditor->OnBlueprintPreCompile().AddUObject(this, &AVirtualCameraActor::OnBlueprintPreCompile);
+			FEditorSupportDelegates::PrepareToCleanseEditorObject.AddUObject(this, &AVirtualCameraActor::OnPrepareToCleanseEditorObject);
+			FAssetRegistryModule& AssetRegistryModule = FModuleManager::LoadModuleChecked<FAssetRegistryModule>(AssetRegistryName);
+			AssetRegistryModule.Get().OnAssetRemoved().AddUObject(this, &AVirtualCameraActor::OnAssetRemoved);
+			FEditorDelegates::OnAssetsCanDelete.AddUObject(this, &AVirtualCameraActor::OnAssetsCanDelete);
+		}
+	}
+	else
+#endif
+	{
+		APlayerController* PlayerController = ActorWorld->GetGameInstance()->GetFirstLocalPlayerController();
+		if (!PlayerController)
+		{
+			return false;
+		}
+
+		PreviousViewTarget = PlayerController->GetViewTarget();
+
+		FViewTargetTransitionParams TransitionParams;
+		PlayerController->SetViewTarget(this, TransitionParams);
+	}
+	
+	UCineCameraComponent* CineCamera = GetCineCameraComponent();
+	if (CineCamera)
+	{
+		// use the aspect ratio of the device we're streaming to, so the UMG and the camera capture fit together and span the device's surface
+		CineCamera->Filmback.SensorWidth = TargetDeviceResolution.X / 100.0f;
+		CineCamera->Filmback.SensorHeight = TargetDeviceResolution.Y / 100.0f;
+	}
+
+	if (CameraUMGClass)
+	{
+		CameraScreenWidget->SetWidgetClass(CameraUMGClass);
+		CameraScreenWidget->Display(ActorWorld);
+	}
+
+	SetActorTickEnabled(true);
+
+	bIsStreaming = true;
+	return true;
+}
+
+bool AVirtualCameraActor::StopStreaming()
+{
+	CameraScreenWidget->Hide();
+	if (MediaCapture)
+	{
+		MediaCapture->StopCapture(true);
+	}
+
+#if WITH_EDITOR
+	if (ActorWorld && ActorWorld->WorldType == EWorldType::Editor)
+	{
+		FLevelEditorModule& LevelEditorModule = FModuleManager::GetModuleChecked<FLevelEditorModule>(LevelEditorName);
+		TSharedPtr<SLevelViewport> ActiveLevelViewport = LevelEditorModule.GetFirstActiveLevelViewport();
+		if (ActiveLevelViewport.IsValid())
+		{
+			// restore FOV
+			GCurrentLevelEditingViewportClient->ViewFOV = GCurrentLevelEditingViewportClient->FOVAngle;
+
+			FLevelEditorViewportClient& LevelViewportClient = ActiveLevelViewport->GetLevelViewportClient();
+			LevelViewportClient.SetActorLock(ViewportSettingsBackup->ActorLock.Get());
+			GCurrentLevelEditingViewportClient->UpdateViewForLockedActor();
+
+			// remove roll and pitch from camera when unbinding from actors
+			GEditor->RemovePerspectiveViewRotation(true, true, false);
+
+			LevelViewportClient.SetRealtime(ViewportSettingsBackup->bRealTime);
+			LevelViewportClient.bDrawAxes = ViewportSettingsBackup->bDrawAxes;
+			LevelViewportClient.bDisableInput = ViewportSettingsBackup->bDisableInput;
+			LevelViewportClient.SetAllowCinematicControl(ViewportSettingsBackup->bAllowCinematicControl);
+
+			// unlock viewport resize
+			ActiveLevelViewport->GetSharedActiveViewport()->SetFixedViewportSize(0, 0);
+
+			// remove event listeners
+			FEditorDelegates::OnAssetsCanDelete.RemoveAll(this);
+			LevelEditorModule.OnMapChanged().RemoveAll(this);
+			if (FAssetRegistryModule* AssetRegistryModule = FModuleManager::GetModulePtr<FAssetRegistryModule>(AssetRegistryName))
+			{
+				IAssetRegistry* AssetRegistry = AssetRegistryModule->TryGet();
+				if (AssetRegistry)
+				{
+					AssetRegistry->OnAssetRemoved().RemoveAll(this);
+				}
+			}
+			FEditorSupportDelegates::PrepareToCleanseEditorObject.RemoveAll(this);
+			GEditor->OnBlueprintPreCompile().RemoveAll(this);
+		}
+
+		ViewportSettingsBackup.Reset();
+	}
+	else
+#endif
+	{
+		if (PreviousViewTarget)
+		{
+			APlayerController* PlayerController = ActorWorld->GetGameInstance()->GetFirstLocalPlayerController();
+			if (!PlayerController)
+			{
+				return false;
+			}
+			FViewTargetTransitionParams TransitionParams;
+			PlayerController->SetViewTarget(PreviousViewTarget, TransitionParams);
+		}
+	}
+
+	SetActorTickEnabled(false);
+
+	bIsStreaming = false;
+
+	if (bSaveSettingsOnStopStreaming)
+	{
+		SaveSettings();
+	}
+
+	return true;
+}
+
+UWorld* AVirtualCameraActor::GetControllerWorld()
+{
+	return GetWorld();
+}
+
+UCineCameraComponent* AVirtualCameraActor::GetStreamedCameraComponent_Implementation() const
+{
+	return GetCineCameraComponent();
+}
+
+USceneCaptureComponent2D * AVirtualCameraActor::GetSceneCaptureComponent_Implementation() const
+{
+	return SceneCaptureComponent;
+}
+
+ULevelSequencePlaybackController* AVirtualCameraActor::GetSequenceController_Implementation() const
+{
+	UVirtualCameraSubsystem* SubSystem = GEngine->GetEngineSubsystem<UVirtualCameraSubsystem>();
+	return SubSystem ? SubSystem->SequencePlaybackController : nullptr;
+}
+
+TScriptInterface<IVirtualCameraPresetContainer> AVirtualCameraActor::GetPresetContainer_Implementation()
+{
+	return this;
+}
+
+TScriptInterface<IVirtualCameraOptions> AVirtualCameraActor::GetOptions_Implementation()
+{
+	return this;
+}
+
+FLiveLinkSubjectRepresentation AVirtualCameraActor::GetLiveLinkRepresentation_Implementation() const
+{
+	return LiveLinkSubject;
+}
+
+void AVirtualCameraActor::SetLiveLinkRepresentation_Implementation(const FLiveLinkSubjectRepresentation& InLiveLinkRepresentation)
+{
+	LiveLinkSubject = InLiveLinkRepresentation;
+}
+
+FString AVirtualCameraActor::SavePreset_Implementation(const bool bSaveCameraSettings, const bool bSaveStabilization, const bool bSaveAxisLocking, const bool bSaveMotionScale)
+{
+	// Convert index to string with leading zeros
+	FString PresetName = FString::Printf(TEXT("Preset-%03i"), PresetIndex);
+
+	// Another preset has been created
+	PresetIndex++;
+	FVirtualCameraSettingsPreset::NextIndex++;
+
+	FVirtualCameraSettingsPreset PresetToAdd;
+	PresetToAdd.DateCreated = FDateTime::UtcNow();
+
+	PresetToAdd.bIsCameraSettingsSaved = bSaveCameraSettings;
+	PresetToAdd.bIsStabilizationSettingsSaved = bSaveStabilization;
+	PresetToAdd.bIsAxisLockingSettingsSaved = bSaveAxisLocking;
+	PresetToAdd.bIsMotionScaleSettingsSaved = bSaveMotionScale;
+
+	UCineCameraComponent* CineCamera = GetCineCameraComponent();
+	if (CineCamera)
+	{
+		PresetToAdd.CameraSettings.FocalLength = CineCamera->CurrentFocalLength;
+		PresetToAdd.CameraSettings.Aperture = CineCamera->CurrentAperture;
+		PresetToAdd.CameraSettings.FilmbackWidth = CineCamera->Filmback.SensorWidth;
+		PresetToAdd.CameraSettings.FilmbackHeight = CineCamera->Filmback.SensorHeight;
+	}
+
+	SettingsPresets.Add(PresetName, PresetToAdd);
+
+	return PresetName;
+}
+
+bool AVirtualCameraActor::LoadPreset_Implementation(const FString& PresetName)
+{
+	FVirtualCameraSettingsPreset* LoadedPreset = SettingsPresets.Find(PresetName);
+
+	if (LoadedPreset)
+	{
+		UCineCameraComponent* CineCamera = GetCineCameraComponent();
+		if (CineCamera)
+		{
+			if (LoadedPreset->bIsCameraSettingsSaved)
+			{
+				CineCamera->CurrentAperture = LoadedPreset->CameraSettings.Aperture;
+				CineCamera->CurrentFocalLength = LoadedPreset->CameraSettings.FocalLength;
+				CineCamera->Filmback.SensorWidth = LoadedPreset->CameraSettings.FilmbackWidth;
+				CineCamera->Filmback.SensorHeight = LoadedPreset->CameraSettings.FilmbackHeight;
+			}
+		}
+		return true;
+	}
+
+	return false;
+}
+
+int32 AVirtualCameraActor::DeletePreset_Implementation(const FString& PresetName)
+{
+	return SettingsPresets.Remove(PresetName);
+}
+
+TMap<FString, FVirtualCameraSettingsPreset> AVirtualCameraActor::GetSettingsPresets_Implementation()
+{
+	SettingsPresets.KeySort([](const FString& a, const FString& b) -> bool
+	{
+		return a < b;
+	});
+
+	return SettingsPresets;
+}
+
+void AVirtualCameraActor::SetDesiredDistanceUnits_Implementation(const EUnit InDesiredUnits)
+{
+	DesiredDistanceUnits = InDesiredUnits;
+}
+
+EUnit AVirtualCameraActor::GetDesiredDistanceUnits_Implementation()
+{
+	return DesiredDistanceUnits;
+}
+
+bool AVirtualCameraActor::IsFocusVisualizationAllowed_Implementation()
+{
+	return bAllowFocusVisualization;
+}
+
+void AVirtualCameraActor::OnTouchEventOutsideUMG(const FVector2D& InViewportPosition)
+{
+	FVector TraceDirection;
+	FVector CameraWorldLocation;
+	if (!DeprojectScreenToWorld(InViewportPosition, CameraWorldLocation, TraceDirection))
+	{
+		return;
+	}
+
+	FCollisionQueryParams TraceParams(SCENE_QUERY_STAT(UpdateAutoFocus), true);
+
+	const FVector TraceEnd = CameraWorldLocation + TraceDirection * MaxFocusTraceDistance;
+	const bool bHit = GetWorld()->LineTraceSingleByChannel(LastViewportTouchResult, CameraWorldLocation, TraceEnd, ECC_Visibility, TraceParams);
+
+	if (bHit)
+	{
+		FEditorScriptExecutionGuard ScriptGuard;
+		OnActorClickedDelegate.ExecuteIfBound(LastViewportTouchResult);
+	}
+}
+
+void AVirtualCameraActor::SaveSettings()
+{
+	UCineCameraComponent* CineCamera = GetCineCameraComponent();
+	if (!CineCamera)
+	{
+		return;
+	}
+
+	UVirtualCameraSaveGame* SaveGameInstance = Cast<UVirtualCameraSaveGame>(UGameplayStatics::CreateSaveGameObject(UVirtualCameraSaveGame::StaticClass()));
+
+
+	// Save focal length and aperture
+	SaveGameInstance->CameraSettings.FocalLength = CineCamera->CurrentFocalLength;
+	SaveGameInstance->CameraSettings.Aperture = CineCamera->CurrentAperture;
+	SaveGameInstance->CameraSettings.bAllowFocusVisualization = bAllowFocusVisualization;
+#if WITH_EDITORONLY_DATA
+	SaveGameInstance->CameraSettings.DebugFocusPlaneColor = CineCamera->FocusSettings.DebugFocusPlaneColor;
+#else
+	SaveGameInstance->CameraSettings.DebugFocusPlaneColor = FColor();
+#endif
+
+	// Save filmback settings
+	SaveGameInstance->CameraSettings.FilmbackName = CineCamera->GetFilmbackPresetName();
+	SaveGameInstance->CameraSettings.FilmbackWidth = CineCamera->Filmback.SensorWidth;
+	SaveGameInstance->CameraSettings.FilmbackHeight = CineCamera->Filmback.SensorHeight;
+
+	// Save settings presets
+	SaveGameInstance->SettingsPresets = SettingsPresets;
+
+	// Save indices for naming
+	SaveGameInstance->WaypointIndex = FVirtualCameraWaypoint::NextIndex;
+	SaveGameInstance->ScreenshotIndex = FVirtualCameraScreenshot::NextIndex;
+	SaveGameInstance->PresetIndex = FVirtualCameraSettingsPreset::NextIndex;
+
+	SaveGameInstance->CameraSettings.DesiredDistanceUnits = DesiredDistanceUnits;
+
+	// Write save file to disk
+	UGameplayStatics::SaveGameToSlot(SaveGameInstance, SavedSettingsSlotName, 0);
+}
+
+void AVirtualCameraActor::LoadSettings()
+{
+	UCineCameraComponent* CineCamera = GetCineCameraComponent();
+	if (!CineCamera)
+	{
+		return;
+	}
+
+	UVirtualCameraSaveGame* SaveGameInstance = Cast<UVirtualCameraSaveGame>(UGameplayStatics::CreateSaveGameObject(UVirtualCameraSaveGame::StaticClass()));
+	SaveGameInstance = Cast<UVirtualCameraSaveGame>(UGameplayStatics::LoadGameFromSlot(SavedSettingsSlotName, 0));
+
+	if (!SaveGameInstance)
+	{
+		UE_LOG(LogVirtualCamera, Warning, TEXT("VirtualCamera could not find save game to load, using default settings."))
+		return;
+	}
+
+	bAllowFocusVisualization = SaveGameInstance->CameraSettings.bAllowFocusVisualization;
+
+#if WITH_EDITORONLY_DATA
+	if (SaveGameInstance->CameraSettings.DebugFocusPlaneColor != FColor())
+	{
+		CineCamera->FocusSettings.DebugFocusPlaneColor = SaveGameInstance->CameraSettings.DebugFocusPlaneColor;
+	}
+#endif
+
+	CineCamera->SetCurrentFocalLength(SaveGameInstance->CameraSettings.FocalLength);
+	CineCamera->CurrentAperture = SaveGameInstance->CameraSettings.Aperture;
+	CineCamera->Filmback.SensorWidth = SaveGameInstance->CameraSettings.FilmbackWidth;
+	CineCamera->Filmback.SensorHeight = SaveGameInstance->CameraSettings.FilmbackHeight;
+
+	DesiredDistanceUnits = SaveGameInstance->CameraSettings.DesiredDistanceUnits;
+
+	// load presets, but don't overwrite existing ones
+	SettingsPresets.Append(SaveGameInstance->SettingsPresets);
+
+	// If the saved preset index is smaller than total presets, set it so that it won't overwrite existing presets.
+	FVirtualCameraSettingsPreset::NextIndex = SaveGameInstance->PresetIndex;
+	if (SettingsPresets.Num() > FVirtualCameraSettingsPreset::NextIndex)
+	{
+		FVirtualCameraSettingsPreset::NextIndex = SettingsPresets.Num();
+	}
+
+	PresetIndex = FVirtualCameraSettingsPreset::NextIndex;
+}
+
+#if WITH_EDITOR
+void AVirtualCameraActor::OnMapChanged(UWorld* World, EMapChangeType ChangeType)
+{
+	if (World == ActorWorld && ChangeType == EMapChangeType::TearDownWorld)
+	{
+		StopStreaming();
+	}
+}
+
+void AVirtualCameraActor::OnBlueprintPreCompile(UBlueprint* Blueprint)
+{
+	if (Blueprint && CameraUMGClass && Blueprint->GeneratedClass == CameraUMGClass)
+	{
+		StopStreaming();
+	}
+}
+
+void AVirtualCameraActor::OnPrepareToCleanseEditorObject(UObject* Object)
+{
+	if (Object == CameraScreenWidget || Object == CameraUMGClass || Object == ActorWorld || Object == MediaCapture)
+	{
+		StopStreaming();
+	}
+}
+
+void AVirtualCameraActor::OnAssetRemoved(const FAssetData& AssetData)
+{
+	if (AssetData.GetPackage() == CameraUMGClass->GetOutermost())
+	{
+		StopStreaming();
+	}
+}
+
+void AVirtualCameraActor::OnAssetsCanDelete(const TArray<UObject*>& InAssetsToDelete, FCanDeleteAssetResult& CanDeleteResult)
+{
+	if (CameraUMGClass)
+	{
+		for (UObject* Obj : InAssetsToDelete)
+		{
+			if (CameraUMGClass->GetOutermost() == Obj->GetOutermost())
+			{
+				UE_LOG(LogVirtualCamera, Warning, TEXT("Asset '%s' can't be deleted because it is currently used by the Virtual Camera Stream."), *Obj->GetPathName());
+				CanDeleteResult.Set(false);
+				break;
+			}
+		}
+	}
+}
+#endif
diff --git a/Engine/Plugins/Experimental/VirtualCamera/Source/VirtualCamera/Private/VirtualCameraCineCameraComponent.cpp b/Engine/Plugins/Experimental/VirtualCamera/Source/VirtualCamera/Private/VirtualCameraCineCameraComponent.cpp
new file mode 100644
--- /dev/null
+++ b/Engine/Plugins/Experimental/VirtualCamera/Source/VirtualCamera/Private/VirtualCameraCineCameraComponent.cpp
@@ -0,0 +1,268 @@
+// Copyright Epic Games, Inc. All Rights Reserved.
+
+#include "VirtualCameraCineCameraComponent.h"
+#include "Engine/StaticMesh.h"
+#include "Kismet/KismetMathLibrary.h"
+#include "UnrealEngine.h"
+
+UVirtualCameraCineCameraComponent::UVirtualCameraCineCameraComponent(const FObjectInitializer& ObjectInitializer)
+{
+	CurrentFilmbackOptionsKey.Empty();
+
+	// Manual focus is the default focus methods
+	SetFocusMethod(EVirtualCameraFocusMethod::Manual);
+	// Default focus distance
+	SetFocusDistance(1000.f);
+	// Default to smooth focus changing
+	FocusSettings.bSmoothFocusChanges = true;
+	// Default smoothing speed
+	FocusSettings.FocusSmoothingInterpSpeed = 100.f;
+	
+	// Constrained aspect ratio is disabled by default, matte is handled by UI
+	bConstrainAspectRatio = false;
+	// Default view size ratio
+	ViewSizeRatio = FVector2D(1.f, 1.f);
+
+	MatteOpacity = .7f;
+	
+	// By default allow camera view updates
+	bAllowCameraViewUpdates = true;
+}
+
+void UVirtualCameraCineCameraComponent::BeginPlay()
+{
+	Super::BeginPlay();
+
+	// Get a default FilmbackKey name
+	TArray<FString> FilmbackKeys;
+	if (GetFilmbackPresetOptions(FilmbackKeys))
+	{
+		CurrentFilmbackOptionsKey = FilmbackKeys[0];
+	}
+
+	if (FilmbackOptions.Contains(CurrentFilmbackOptionsKey))
+	{
+		DesiredFilmbackSettings = FilmbackOptions[CurrentFilmbackOptionsKey];
+	}
+
+	SetMatteAspectRatio(DesiredFilmbackSettings.SensorWidth / DesiredFilmbackSettings.SensorHeight);
+}
+
+float UVirtualCameraCineCameraComponent::ChangeFocalLengthPreset(const bool bShiftUp)
+{
+	if (FocalLengthOptions.Num() > 0)
+	{
+		int32 TargetIndex = FindClosestPresetIndex(FocalLengthOptions, GetCurrentFocalLength());
+
+		// If current focal length is a preset, a move still needs to occur when switching
+		if (UKismetMathLibrary::NearlyEqual_FloatFloat(FocalLengthOptions[TargetIndex], GetCurrentFocalLength()))
+		{
+			TargetIndex += bShiftUp ? 1 : -1;
+		}
+		// If should shift down and target is above current, move down one
+		else if (!bShiftUp && FocalLengthOptions[TargetIndex] > GetCurrentFocalLength())
+		{
+			TargetIndex -= 1;
+		}
+		// If should shift up and target is below current, move up one
+		else if (bShiftUp && FocalLengthOptions[TargetIndex] < GetCurrentFocalLength())
+		{
+			TargetIndex += 1;
+		}
+
+		TargetIndex = FMath::Clamp<int32>(TargetIndex, 0, FocalLengthOptions.Num() - 1);
+
+		SetCurrentFocalLength(FocalLengthOptions[TargetIndex]);
+		return GetCurrentFocalLength();
+	}
+
+	// If there are no presets, a value still needs to be returned.
+	// A negative value communicates that there are no presets, but everything is still functioning normally.
+	return -1.f;
+}
+
+float UVirtualCameraCineCameraComponent::ChangeAperturePreset(const bool bShiftUp)
+{
+	if (ApertureOptions.Num() > 0)
+	{
+		int32 TargetIndex = FindClosestPresetIndex(ApertureOptions, GetCurrentAperture());
+
+		// If current aperture is a preset, a move still needs to occur when switching
+		if (UKismetMathLibrary::NearlyEqual_FloatFloat(ApertureOptions[TargetIndex], GetCurrentAperture()))
+		{
+			TargetIndex += bShiftUp ? 1 : -1;
+		}
+		// If should shift down and target is above current, move down one
+		else if (!bShiftUp && ApertureOptions[TargetIndex] > GetCurrentAperture())
+		{
+			TargetIndex -= 1;
+		}
+		// If should shift up and target is below current, move up one
+		else if (bShiftUp && ApertureOptions[TargetIndex] < GetCurrentAperture())
+		{
+			TargetIndex += 1;
+		}
+
+		TargetIndex = FMath::Clamp<int32>(TargetIndex, 0, ApertureOptions.Num() - 1);
+
+		SetCurrentAperture(ApertureOptions[TargetIndex]);
+		return GetCurrentFocalLength();
+	}
+
+	// If there are no presets, a value still needs to be returned.
+	// A negative value communicates that there are no presets, but everything is still functioning normally.
+	return -1.f;
+}
+
+bool UVirtualCameraCineCameraComponent::GetFilmbackPresetOptions(TArray<FString>& OutFilmbackPresetsArray) const
+{
+	int32 NumKeys = FilmbackOptions.GetKeys(OutFilmbackPresetsArray);
+	return NumKeys > 0;
+}
+
+bool UVirtualCameraCineCameraComponent::SetFilmbackPresetOption(const FString& NewFilmbackPreset)
+{
+	// Search for filmback in filmback array and return it
+	if (FilmbackOptions.Contains(NewFilmbackPreset))
+	{
+		DesiredFilmbackSettings = FilmbackOptions[NewFilmbackPreset];
+		CurrentFilmbackOptionsKey = NewFilmbackPreset;
+		SetMatteAspectRatio(DesiredFilmbackSettings.SensorWidth / DesiredFilmbackSettings.SensorHeight);
+		return true;
+	}
+
+	return false;
+}
+
+void UVirtualCameraCineCameraComponent::GetMatteValues(TArray<float>& OutMatteValues) const 
+{
+	OutMatteValues = MatteOptions;
+}
+
+void UVirtualCameraCineCameraComponent::UpdateCameraView()
+{
+	// Check to make sure frames should be updated
+	// Scenarios where updates are unwanted include taking a screenshot or recording a sequence
+	if (bAllowCameraViewUpdates && FilmbackOptions.Contains(CurrentFilmbackOptionsKey))
+	{
+		DesiredFilmbackSettings = FilmbackOptions[CurrentFilmbackOptionsKey];
+		DesiredFilmbackSettings.SensorAspectRatio = DesiredFilmbackSettings.SensorWidth / DesiredFilmbackSettings.SensorHeight;
+		float ViewSizeAdjustmentForMatte = MatteAspectRatio / DesiredFilmbackSettings.SensorAspectRatio;
+
+		// Set the actual camera filmback settings
+		if (ViewSizeAdjustmentForMatte >= 1.f)
+		{
+			Filmback.SensorWidth = DesiredFilmbackSettings.SensorWidth / ViewSizeRatio.X;
+			Filmback.SensorHeight = DesiredFilmbackSettings.SensorHeight / ViewSizeRatio.Y;
+		}
+		else
+		{
+			Filmback.SensorWidth = DesiredFilmbackSettings.SensorWidth / ViewSizeRatio.X * ViewSizeAdjustmentForMatte;
+			Filmback.SensorHeight = DesiredFilmbackSettings.SensorHeight / ViewSizeRatio.Y * ViewSizeAdjustmentForMatte;
+		}
+	}
+}
+
+bool UVirtualCameraCineCameraComponent::SetMatteAspectRatio(const float NewMatteAspectRatio)
+{
+	if (NewMatteAspectRatio > 0)
+	{
+		MatteAspectRatio = NewMatteAspectRatio;
+		return true;
+	}
+
+	return false;
+}
+
+void UVirtualCameraCineCameraComponent::SetFocusMethod(const EVirtualCameraFocusMethod NewFocusMethod)
+{
+	CurrentFocusMethod = NewFocusMethod;
+
+	switch (NewFocusMethod)
+	{
+	case EVirtualCameraFocusMethod::None:
+		bAutoFocusEnabled = false;
+		FocusSettings.FocusMethod = ECameraFocusMethod::DoNotOverride;
+		break;
+	case EVirtualCameraFocusMethod::Auto:
+		bAutoFocusEnabled = true;
+		FocusSettings.FocusMethod = ECameraFocusMethod::Manual;
+		break;
+	case EVirtualCameraFocusMethod::Manual:
+		bAutoFocusEnabled = false;
+		FocusSettings.FocusMethod = ECameraFocusMethod::Manual;
+		break;
+	case EVirtualCameraFocusMethod::Tracking:
+		bAutoFocusEnabled = false;
+		FocusSettings.FocusMethod = ECameraFocusMethod::Tracking;
+		break;
+	default:  // Should never be reached, but just in case new focus methods are added
+		UE_LOG(LogActor, Warning, TEXT("Specified focus method is not currently supported in Virtual Camera!"))
+		break;
+	}
+}
+
+void UVirtualCameraCineCameraComponent::SetFocusDistance(const float NewFocusDistance)
+{
+	FocusSettings.ManualFocusDistance = NewFocusDistance;
+	FocusSettings.FocusOffset = 0.f;
+}
+
+void UVirtualCameraCineCameraComponent::SetTrackedActorForFocus(AActor* ActorToTrack, const FVector TrackingPointOffset)
+{
+	FocusSettings.TrackingFocusSettings.ActorToTrack = ActorToTrack;
+	FocusSettings.TrackingFocusSettings.RelativeOffset = TrackingPointOffset;
+}
+
+void UVirtualCameraCineCameraComponent::AddBlendableToCamera(TScriptInterface<IBlendableInterface> BlendableToAdd, float Weight)
+{
+	PostProcessSettings.AddBlendable(BlendableToAdd, Weight);
+}
+
+void UVirtualCameraCineCameraComponent::SetFocusChangeSmoothness(const float NewFocusChangeSmoothness)
+{
+	FocusChangeSmoothness = FMath::Clamp<float>(NewFocusChangeSmoothness, 0.f, 1.f);
+
+	// Translate the user-facing value into the actual interpolation speed
+	FocusSettings.FocusSmoothingInterpSpeed = FMath::Lerp(10000.f, 1.f, NewFocusChangeSmoothness);
+}
+
+void UVirtualCameraCineCameraComponent::SetFocusVisualization(bool bShowFocusVisualization)
+{
+#if WITH_EDITOR
+	if (FocusSettings.FocusMethod == ECameraFocusMethod::Disable)
+	{
+		UE_LOG(LogActor, Warning, TEXT("Camera focus mode is currently set to none, cannot display focus plane!"))
+		return;
+	}
+
+	FocusSettings.bDrawDebugFocusPlane = bShowFocusVisualization;
+#endif
+}
+
+bool UVirtualCameraCineCameraComponent::IsFocusVisualizationActivated() const
+{
+#if WITH_EDITOR
+	return FocusSettings.FocusMethod != ECameraFocusMethod::Disable && FocusSettings.bDrawDebugFocusPlane;
+#else
+	return false;
+#endif
+}
+
+int32 UVirtualCameraCineCameraComponent::FindClosestPresetIndex(const TArray<float>& ArrayToSearch, const float SearchValue) const
+{
+	float LowestDifference = 0.f;
+	int32 ClosestPresetIndex = 0;
+
+	for (int32 Index = 0; Index < ArrayToSearch.Num(); Index++)
+	{
+		float Difference = FMath::Abs(ArrayToSearch[Index] - SearchValue);
+		if (Difference < LowestDifference || Index == 0)
+		{
+			LowestDifference = Difference;
+			ClosestPresetIndex = Index;
+		}
+	}
+
+	return ClosestPresetIndex;
+}
diff --git a/Engine/Plugins/Experimental/VirtualCamera/Source/VirtualCamera/Private/VirtualCameraConcertCameraComponent.cpp b/Engine/Plugins/Experimental/VirtualCamera/Source/VirtualCamera/Private/VirtualCameraConcertCameraComponent.cpp
new file mode 100644
--- /dev/null
+++ b/Engine/Plugins/Experimental/VirtualCamera/Source/VirtualCamera/Private/VirtualCameraConcertCameraComponent.cpp
@@ -0,0 +1,72 @@
+// Copyright Epic Games, Inc. All Rights Reserved.
+
+#include "VirtualCameraConcertCameraComponent.h"
+
+#include "CineCameraComponent.h"
+#include "ConcertVirtualCamera.h"
+#include "VirtualCamera.h"
+
+
+UDEPRECATED_VirtualCameraConcertCameraComponent::UDEPRECATED_VirtualCameraConcertCameraComponent(const FObjectInitializer& ObjectInitializer)
+{ 
+	PrimaryComponentTick.bCanEverTick = true;
+	PrimaryComponentTick.TickGroup = TG_DuringPhysics;
+	PrimaryComponentTick.bTickEvenWhenPaused = true;
+
+	bUpdateCameraComponentTransform = true;
+}
+
+
+void UDEPRECATED_VirtualCameraConcertCameraComponent::BeginPlay()
+{
+	Super::BeginPlay();
+
+	if (ComponentToTrack == nullptr)
+	{
+		if (GetOwner())
+		{
+			ComponentToTrack = GetOwner()->FindComponentByClass<UCineCameraComponent>();
+		}
+	}
+
+	if (ComponentToTrack == nullptr)
+	{
+		UE_LOG(LogVirtualCamera, Error, TEXT("There is no camera to track for %s"), *GetFullName());
+	}
+	else if (ComponentToTrack->GetOwner())
+	{
+		UE_LOG(LogVirtualCamera, Warning, TEXT("The track camera %s doesn't have an owner."), *ComponentToTrack->GetFullName());
+	}
+
+	if (TrackingName.IsNone())
+	{
+		UE_LOG(LogVirtualCamera, Error, TEXT("%s doesn't have a tacking name."));
+	}
+}
+
+
+void UDEPRECATED_VirtualCameraConcertCameraComponent::TickComponent(float DeltaTime, enum ELevelTick TickType, FActorComponentTickFunction *ThisTickFunction)
+{
+	Super::TickComponent(DeltaTime, TickType, ThisTickFunction);
+
+	if (ComponentToTrack == nullptr || ComponentToTrack->GetOwner() == nullptr)
+	{
+		return;
+	}
+
+	FName CurrentTrackingName = TrackingName.IsNone() ? GetFName() : TrackingName;
+
+	if (!bHasAuthority)
+	{
+		FConcertVirtualCameraCameraData CameraData;
+		if (IVirtualCameraModule::Get().GetConcertVirtualCameraManager()->GetLatestCameraComponentEvent(CurrentTrackingName, CameraData))
+		{
+			CameraData.ApplyTo(ComponentToTrack->GetOwner(), ComponentToTrack, bUpdateCameraComponentTransform);
+		}
+	}
+	else if (TickType != LEVELTICK_ViewportsOnly || bSendUpdateInEditor)
+	{
+		FConcertVirtualCameraCameraData CameraData = FConcertVirtualCameraCameraData(ComponentToTrack->GetOwner(), ComponentToTrack);
+		IVirtualCameraModule::Get().GetConcertVirtualCameraManager()->SendCameraCompoentEvent(CurrentTrackingName, CameraData);
+	}
+}
diff --git a/Engine/Plugins/Experimental/VirtualCamera/Source/VirtualCamera/Private/VirtualCameraMovementComponent.cpp b/Engine/Plugins/Experimental/VirtualCamera/Source/VirtualCamera/Private/VirtualCameraMovementComponent.cpp
new file mode 100644
--- /dev/null
+++ b/Engine/Plugins/Experimental/VirtualCamera/Source/VirtualCamera/Private/VirtualCameraMovementComponent.cpp
@@ -0,0 +1,357 @@
+// Copyright Epic Games, Inc. All Rights Reserved.
+
+#include "VirtualCameraMovementComponent.h"
+
+#include "Components/PrimitiveComponent.h"
+#include "Kismet/KismetMathLibrary.h"
+
+LLM_DEFINE_TAG(VirtualCamera_VirtualCameraMovementComponent);
+
+namespace
+{
+	float ClampStabilizationScale(float StabilizationScale)
+	{
+		const float StabilizationMin = 0.0f;
+		const float StabilizationMax = 0.97f;
+		const float StabilizationExp = 0.25f;
+
+		// Stabilization is applied as an exponential curve defined by StabilizationMax
+		// This exponent is less than 1 so value changes at numbers closer to 0 have a greater degree of change than those closer to 1.
+		// This has the effect of stabilization being introduced much faster when it is first applied, and giving a finer degree of control as values increase.
+		// This is necessary because stabilization is not noticeable or useful until a certain amount is applied.
+
+		float NewStabilizationAmountRotationX = FMath::Clamp<float>(StabilizationScale, StabilizationMin, StabilizationMax);
+		return FMath::Pow(NewStabilizationAmountRotationX, StabilizationExp) * StabilizationMax;
+	}
+}
+
+UVirtualCameraMovementComponent::UVirtualCameraMovementComponent(const FObjectInitializer& ObjectInitializer)
+{
+	LLM_SCOPE_BYTAG(VirtualCamera_VirtualCameraMovementComponent);
+
+	// Initialize Axis Settings Map
+	AxisSettings.Add(EVirtualCameraAxis::LocationX);
+	AxisSettings.Add(EVirtualCameraAxis::LocationY);
+	AxisSettings.Add(EVirtualCameraAxis::LocationZ);
+	AxisSettings.Add(EVirtualCameraAxis::RotationX);
+	AxisSettings.Add(EVirtualCameraAxis::RotationY);
+	AxisSettings.Add(EVirtualCameraAxis::RotationZ);
+
+	bLockRelativeToFirstLockAxis = true;
+	bUseGlobalBoom = true;
+	bZeroDutchOnLock = true;
+	CachedLockingAxis = FQuat::Identity;
+}
+
+void UVirtualCameraMovementComponent::AddInputVector(FVector WorldVector, bool bForce /*=false*/)
+{
+	if (WorldVector.IsZero())
+	{
+		return;
+	}
+
+	ApplyLocationScaling(WorldVector);
+	ApplyLocationLocks(WorldVector);
+
+	TargetLocation += WorldVector;
+}
+
+void UVirtualCameraMovementComponent::AddInputVectorFromController(FVector WorldVector, EVirtualCameraAxis MovementScaleAxis)
+{
+	if (WorldVector.IsZero())
+	{
+		return;
+	}
+
+	WorldVector *= AxisSettings[MovementScaleAxis].MovementScale;
+
+	// For Controller movement, move the RootUpdatedComponent
+	FromControllerTargetLocation += WorldVector;
+}
+
+void UVirtualCameraMovementComponent::ProcessMovementInput(const FVector& TrackerLocation, const FRotator& TrackerRotation)
+{
+	FVector DeltaMovement = TrackerLocation - PreviousTrackerLocation;
+	FRotator RotationReorient = GetRotationOffset();
+	RotationReorient.Pitch = 0.f;
+	// make sure movement is correctly oriented to current direction camera is facing
+	DeltaMovement = RotationReorient.UnrotateVector(DeltaMovement);
+	AddInputVector(DeltaMovement);
+	PreviousTrackerLocation = TrackerLocation;
+
+	// Calculate the global change in rotation and apply it
+	FRotator DeltaRotation = UKismetMathLibrary::NormalizedDeltaRotator(TrackerRotation, PreviousTrackerRotation);
+	ApplyRotationLocks(DeltaRotation);
+	TargetRotation = TrackerRotation - GetRotationOffset() + GetOwner()->GetActorRotation();
+	PreviousTrackerRotation = TrackerRotation;
+
+	{
+		FHitResult OutResult = FHitResult();
+		SafeMoveUpdatedComponent(GetStabilizedDeltaLocation(), GetStabilizedRotation(), false, OutResult);
+	}
+
+	// Apply the location change for the RootUpdatedComponent
+	if (RootUpdatedComponent)
+	{
+		TGuardValue<decltype(UpdatedComponent)> UpdateComponentGuard(UpdatedComponent, RootUpdatedComponent);
+		TGuardValue<decltype(UpdatedPrimitive)> UpdatePrimitiveGuard(UpdatedPrimitive, RootUpdatedPrimitive);
+		FHitResult OutResult = FHitResult();
+		SafeMoveUpdatedComponent(FromControllerTargetLocation, FQuat::Identity, false, OutResult);
+		FromControllerTargetLocation = FVector::ZeroVector;
+	}
+}
+
+bool UVirtualCameraMovementComponent::ToggleAxisLock(const EVirtualCameraAxis AxisToToggle)
+{
+	const bool WasLocked = IsLocationLockingActive();
+	const bool bIsLocked = AxisSettings[AxisToToggle].ToggleLock();
+
+	if (AxisToToggle == EVirtualCameraAxis::RotationX && bZeroDutchOnLock && bIsLocked)
+	{
+		// This needs to be applied to the lock offset directly to avoid adding to freeze
+		AxisSettings[EVirtualCameraAxis::RotationX].LockRotationOffset = TargetRotation.Roll;
+	}
+
+	// Update cached locking axis if coming from completely unlocked state
+	if (bLockRelativeToFirstLockAxis && !WasLocked)
+	{
+		CachedLockingAxis = UpdatedComponent->GetComponentRotation().Quaternion();
+	}
+
+	return bIsLocked;
+}
+
+bool UVirtualCameraMovementComponent::ToggleAxisFreeze(const EVirtualCameraAxis AxisToToggle)
+{
+	// Leaving this out of header file for the inevitable refactor that will cause this to be more that one line of code.
+	return AxisSettings[AxisToToggle].bIsFrozen = !AxisSettings[AxisToToggle].bIsFrozen;
+}
+
+float UVirtualCameraMovementComponent::SetAxisStabilizationScale(const EVirtualCameraAxis AxisToAdjust, float NewStabilizationAmount)
+{
+	AxisSettings[AxisToAdjust].StabilizationScale = FMath::Clamp<float>(NewStabilizationAmount, 0.0f, 1.0f);
+	return AxisSettings[AxisToAdjust].StabilizationScale;
+}
+
+float UVirtualCameraMovementComponent::GetAxisStabilizationScale(const EVirtualCameraAxis AxisToRetrieve) const
+{
+	return AxisSettings[AxisToRetrieve].StabilizationScale;
+}
+
+void UVirtualCameraMovementComponent::ResetCameraOffsetsToTracker()
+{
+	// Reorient pawn to be in "stage space", with stage origin acting as pawn root component
+	TargetLocation = PreviousTrackerLocation;
+	TargetRotation = PreviousTrackerRotation;
+	UpdatedComponent->SetRelativeLocation(TargetLocation);
+	UpdatedComponent->SetRelativeRotation(TargetRotation);
+
+	// Clear all locks
+	for (auto Iterator = AxisSettings.CreateIterator(); Iterator; ++Iterator)
+	{
+		Iterator.Value().SetIsLocked(false);
+		Iterator.Value().bIsFrozen = false;
+		Iterator.Value().FreezeRotationOffset = 0.f;
+		Iterator.Value().FreezeLocationOffset = FVector::ZeroVector;
+	}
+
+	OnOffsetReset.Broadcast();
+}
+
+bool UVirtualCameraMovementComponent::IsLocationLockingActive() const 
+{
+	return AxisSettings[EVirtualCameraAxis::LocationX].bIsLocked || AxisSettings[EVirtualCameraAxis::LocationY].bIsLocked || AxisSettings[EVirtualCameraAxis::LocationZ].bIsLocked;
+}
+
+bool UVirtualCameraMovementComponent::IsAxisLocked(const EVirtualCameraAxis AxisToCheck) const
+{
+	return AxisSettings[AxisToCheck].bIsLocked;
+}
+
+void UVirtualCameraMovementComponent::OnMoveForward(const float InValue)
+{
+	FVector InputVector = UpdatedComponent->GetForwardVector();
+	InputVector = GetOwner()->GetActorRotation().UnrotateVector(InputVector);
+	AddInputVectorFromController(InputVector * InValue, EVirtualCameraAxis::LocationX);
+}
+
+void UVirtualCameraMovementComponent::OnMoveRight(const float InValue)
+{
+	FVector InputVector = UpdatedComponent->GetRightVector();
+	InputVector = GetOwner()->GetActorRotation().UnrotateVector(InputVector);
+	AddInputVectorFromController(InputVector * InValue, EVirtualCameraAxis::LocationY);
+}
+
+void UVirtualCameraMovementComponent::OnMoveUp(const float InValue)
+{
+	if (bUseGlobalBoom)
+	{
+		FVector InputVector = FVector(0.f, 0.f, 1.f);
+		AddInputVectorFromController(InputVector * InValue, EVirtualCameraAxis::LocationZ);
+	}
+	else
+	{
+		FVector InputVector = UpdatedComponent->GetUpVector();
+		InputVector = GetOwner()->GetActorRotation().UnrotateVector(InputVector);
+		AddInputVectorFromController(InputVector * InValue, EVirtualCameraAxis::LocationZ);
+	}
+}
+
+void UVirtualCameraMovementComponent::Teleport(const FTransform& TargetTransform)
+{
+	ResetCameraOffsetsToTracker();
+
+	// The change between where we were and where we are teleporting to.
+	FVector DeltaOffset = TargetTransform.GetLocation() - UpdatedComponent->GetComponentLocation();
+	DeltaOffset = GetOwner()->GetActorRotation().UnrotateVector(DeltaOffset);
+
+	UpdatedComponent->AddLocalOffset(DeltaOffset);
+	TargetLocation += DeltaOffset;
+
+	if (RootUpdatedComponent)
+	{
+		RootUpdatedComponent->SetRelativeLocation(FVector::ZeroVector);
+		FromControllerTargetLocation = FVector::ZeroVector;
+	}
+}
+
+void UVirtualCameraMovementComponent::SetRootComponent(USceneComponent* InFromController)
+{
+	RootUpdatedComponent = InFromController;
+	RootUpdatedPrimitive = Cast<UPrimitiveComponent>(RootUpdatedComponent);
+}
+
+void UVirtualCameraMovementComponent::ApplyLocationLocks(const FVector& InVector)
+{
+	FVector ForwardVector;
+	FVector RightVector;
+	FVector UpVector;
+
+	GetDirectionVectorsForCamera(ForwardVector, RightVector, UpVector, true);
+
+	FVector AdjustedVectorX = InVector - FVector::VectorPlaneProject(InVector, ForwardVector);
+	FVector AdjustedVectorY = InVector - FVector::VectorPlaneProject(InVector, RightVector);
+	FVector AdjustedVectorZ = InVector - FVector::VectorPlaneProject(InVector, UpVector);
+
+	if (AxisSettings[EVirtualCameraAxis::LocationX].IsAxisImmobilized())
+	{
+		AxisSettings[EVirtualCameraAxis::LocationX].AddLocationOffset(AdjustedVectorX);
+	}
+
+	if (AxisSettings[EVirtualCameraAxis::LocationY].IsAxisImmobilized())
+	{
+		AxisSettings[EVirtualCameraAxis::LocationY].AddLocationOffset(AdjustedVectorY);
+	}
+
+	if (AxisSettings[EVirtualCameraAxis::LocationZ].IsAxisImmobilized())
+	{
+		AxisSettings[EVirtualCameraAxis::LocationZ].AddLocationOffset(AdjustedVectorZ);
+	}
+}
+
+FVector UVirtualCameraMovementComponent::GetStabilizedDeltaLocation() const 
+{
+	FVector ReturnVector = FVector::ZeroVector;
+	FVector ClosestPoint;
+
+	FVector TargetLocationWithOffsets = TargetLocation - GetLocationOffset();
+
+	const FVector UpdatedComponentRelativeLocation = UpdatedComponent->GetRelativeLocation();
+
+	// Calculate each component by getting the needed vector component and doing a lerp with the stabilization scale for that axis
+	FMath::PointDistToLine(TargetLocationWithOffsets, UpdatedComponent->GetForwardVector(), UpdatedComponentRelativeLocation, ClosestPoint);
+	ReturnVector += (ClosestPoint - UpdatedComponentRelativeLocation) * (1 - ClampStabilizationScale(AxisSettings[EVirtualCameraAxis::LocationX].StabilizationScale));
+
+	FMath::PointDistToLine(TargetLocationWithOffsets, UpdatedComponent->GetRightVector(), UpdatedComponentRelativeLocation, ClosestPoint);
+	ReturnVector += (ClosestPoint - UpdatedComponentRelativeLocation) * (1 - ClampStabilizationScale(AxisSettings[EVirtualCameraAxis::LocationY].StabilizationScale));
+
+	FMath::PointDistToLine(TargetLocationWithOffsets, UpdatedComponent->GetUpVector(), UpdatedComponentRelativeLocation, ClosestPoint);
+	ReturnVector += (ClosestPoint - UpdatedComponentRelativeLocation) * (1 - ClampStabilizationScale(AxisSettings[EVirtualCameraAxis::LocationZ].StabilizationScale));
+
+	ReturnVector = GetOwner()->GetActorRotation().RotateVector(ReturnVector);
+	return ReturnVector;
+}
+
+void UVirtualCameraMovementComponent::ApplyLocationScaling(FVector& VectorToAdjust)
+{
+	// Get the axes to scale along
+	FVector ForwardVector;
+	FVector RightVector;
+	FVector UpVector;
+
+	GetDirectionVectorsForCamera(ForwardVector, RightVector, UpVector, false);
+
+	// Orient to global Z up, but maintain yaw
+	ForwardVector = FVector::VectorPlaneProject(ForwardVector, FVector::UpVector);
+	RightVector = FVector::VectorPlaneProject(RightVector, FVector::UpVector);
+	UpVector = FVector::UpVector;
+
+	ForwardVector.Normalize();
+	RightVector.Normalize();
+
+	FVector XComponent = VectorToAdjust.ProjectOnTo(ForwardVector);
+	FVector YComponent = VectorToAdjust.ProjectOnTo(RightVector);
+	FVector ZComponent = VectorToAdjust.ProjectOnTo(UpVector);
+
+	VectorToAdjust = XComponent * AxisSettings[EVirtualCameraAxis::LocationX].MovementScale;
+	VectorToAdjust += YComponent * AxisSettings[EVirtualCameraAxis::LocationY].MovementScale;
+	VectorToAdjust += ZComponent * AxisSettings[EVirtualCameraAxis::LocationZ].MovementScale;
+}
+
+void UVirtualCameraMovementComponent::ApplyRotationLocks(const FRotator& InRotation)
+{
+	// The AddOffset function will ignore incoming offsets if the axis is not immobilized
+	AxisSettings[EVirtualCameraAxis::RotationX].AddRotationOffset(InRotation.Roll);
+	AxisSettings[EVirtualCameraAxis::RotationY].AddRotationOffset(InRotation.Pitch);
+	AxisSettings[EVirtualCameraAxis::RotationZ].AddRotationOffset(InRotation.Yaw);
+}
+
+FRotator UVirtualCameraMovementComponent::GetStabilizedRotation() const
+{
+	// Needs to be done as Quaternions to avoid gimbal locking
+	FRotator TargetAdjustment = UKismetMathLibrary::ComposeRotators(TargetRotation, UpdatedComponent->GetComponentRotation().GetInverse());
+
+	TargetAdjustment.Roll  *= (1 - ClampStabilizationScale(AxisSettings[EVirtualCameraAxis::RotationX].StabilizationScale));
+	TargetAdjustment.Pitch *= (1 - ClampStabilizationScale(AxisSettings[EVirtualCameraAxis::RotationY].StabilizationScale));
+	TargetAdjustment.Yaw   *= (1 - ClampStabilizationScale(AxisSettings[EVirtualCameraAxis::RotationZ].StabilizationScale));
+
+	return UKismetMathLibrary::ComposeRotators(TargetAdjustment, UpdatedComponent->GetComponentRotation());
+}
+
+void UVirtualCameraMovementComponent::GetDirectionVectorsForCamera(FVector& OutForwardVector, FVector& OutRightVector, FVector& OutUpVector, bool bTryLock) const
+{
+	// Break the vector into local components so we can apply individual scaling
+	if (bTryLock && bLockRelativeToFirstLockAxis && IsLocationLockingActive())
+	{
+		// If needed, use the cached axis rather than always using local
+		OutForwardVector = CachedLockingAxis.GetForwardVector();
+		OutRightVector = CachedLockingAxis.GetRightVector();
+		OutUpVector = CachedLockingAxis.GetUpVector();
+	}
+	else
+	{
+		OutForwardVector = UpdatedComponent->GetForwardVector();
+		OutRightVector = UpdatedComponent->GetRightVector();
+		OutUpVector = UpdatedComponent->GetUpVector();
+	}
+
+	OutForwardVector = GetOwner()->GetActorRotation().UnrotateVector(OutForwardVector);
+	OutRightVector = GetOwner()->GetActorRotation().UnrotateVector(OutRightVector);
+	OutUpVector = GetOwner()->GetActorRotation().UnrotateVector(OutUpVector);
+}
+
+FVector UVirtualCameraMovementComponent::GetLocationOffset() const 
+{
+	return AxisSettings[EVirtualCameraAxis::LocationX].GetLocationOffset()
+		+ AxisSettings[EVirtualCameraAxis::LocationY].GetLocationOffset()
+		+ AxisSettings[EVirtualCameraAxis::LocationZ].GetLocationOffset();
+}
+
+FRotator UVirtualCameraMovementComponent::GetRotationOffset() const 
+{
+	return FRotator(
+		AxisSettings[EVirtualCameraAxis::RotationY].GetRotationOffset(),
+		AxisSettings[EVirtualCameraAxis::RotationZ].GetRotationOffset(),
+		AxisSettings[EVirtualCameraAxis::RotationX].GetRotationOffset()
+	);
+}
diff --git a/Engine/Plugins/Experimental/VirtualCamera/Source/VirtualCamera/Private/VirtualCameraPawnBase.cpp b/Engine/Plugins/Experimental/VirtualCamera/Source/VirtualCamera/Private/VirtualCameraPawnBase.cpp
new file mode 100644
--- /dev/null
+++ b/Engine/Plugins/Experimental/VirtualCamera/Source/VirtualCamera/Private/VirtualCameraPawnBase.cpp
@@ -0,0 +1,595 @@
+// Copyright Epic Games, Inc. All Rights Reserved.
+
+#include "VirtualCameraPawnBase.h"
+#include "UnrealClient.h"
+#include "UnrealEngine.h"
+#include "Kismet/GameplayStatics.h"
+
+LLM_DEFINE_TAG(VirtualCamera_VirtualCameraPawnBase);
+
+int32 AVirtualCameraPawnBase::PresetIndex = 1;
+int32 AVirtualCameraPawnBase::WaypointIndex = 1;
+int32 AVirtualCameraPawnBase::ScreenshotIndex = 1;
+
+AVirtualCameraPawnBase::AVirtualCameraPawnBase(const FObjectInitializer& ObjectInitializer)
+{
+	LLM_SCOPE_BYTAG(VirtualCamera_VirtualCameraPawnBase);
+	PrimaryActorTick.bCanEverTick = true;
+	PrimaryActorTick.bStartWithTickEnabled = true;
+
+	// Create Components
+	DefaultSceneRoot = CreateDefaultSubobject<USceneComponent>("DefaultSceneRoot");
+	SetRootComponent(DefaultSceneRoot);
+
+	CineCamera = CreateDefaultSubobject<UVirtualCameraCineCameraComponent>("Cinematic Camera");
+	CineCamera->SetupAttachment(DefaultSceneRoot);
+
+	MovementComponent = CreateDefaultSubobject<UVirtualCameraMovementComponent>("Movement Component");
+	MovementComponent->UpdatedComponent = CineCamera;
+	MovementComponent->SetRootComponent(DefaultSceneRoot);
+
+	// Set Default Variables
+	HomeWaypointName = "";
+
+	// Set formatter to pad with leading zeros to make integers 3 characters long.
+	MinimumIntegralDigits = 3;
+
+	SavedSettingsSlotName = "SavedVirtualCameraSettings";
+	bSaveSettingsWhenClosing = false;
+	DesiredDistanceUnits = EUnit::Meters;
+
+	// By default, allow focus visualization
+	bAllowFocusVisualization = true;
+
+	FScreenshotRequest::OnScreenshotRequestProcessed().AddUObject(ToRawPtr(CineCamera), &UVirtualCameraCineCameraComponent::AllowCameraViewUpdates);
+}
+
+void AVirtualCameraPawnBase::BeginPlay()
+{
+	Super::BeginPlay();
+
+	if (bSaveSettingsWhenClosing)
+	{
+		LoadSettings();
+	}
+
+	LoadFinished();
+}
+
+void AVirtualCameraPawnBase::EndPlay(const EEndPlayReason::Type EndPlayReason)
+{
+	Super::EndPlay(EndPlayReason);
+
+	if (bSaveSettingsWhenClosing)
+	{
+		SaveSettings();
+	}
+}
+
+void AVirtualCameraPawnBase::GetWaypointInfo(const FString WaypointName, FVirtualCameraWaypoint &OutWaypointInfo) const 
+{
+	const FVirtualCameraWaypoint* WaypointInfoPtr = Waypoints.Find(WaypointName);
+	if (WaypointInfoPtr)
+	{
+		OutWaypointInfo = *WaypointInfoPtr;
+	}
+}
+
+void AVirtualCameraPawnBase::GetWaypointNames(TArray<FString>& OutArray) const 
+{
+	Waypoints.GetKeys(OutArray);
+	OutArray.Sort();
+}
+
+void AVirtualCameraPawnBase::GetScreenshotInfo(const FString ScreenshotName, FVirtualCameraScreenshot &OutScreenshotInfo) const 
+{
+	const FVirtualCameraScreenshot* ScreenshotInfoPtr = Screenshots.Find(ScreenshotName);
+	if (ScreenshotInfoPtr)
+	{
+		OutScreenshotInfo = *ScreenshotInfoPtr;
+	}
+}
+
+void AVirtualCameraPawnBase::GetScreenshotNames(TArray<FString>& OutArray) const
+{
+	Screenshots.GetKeys(OutArray);
+	OutArray.Sort();
+}
+
+bool AVirtualCameraPawnBase::RenameWaypoint(const FString& TargetWaypoint, const FString& NewWaypointName)
+{
+	if (Waypoints.Contains(TargetWaypoint))
+	{
+		Waypoints.FindOrAdd(NewWaypointName) = Waypoints[TargetWaypoint];
+		Waypoints[NewWaypointName].Name = NewWaypointName;
+		Waypoints.Remove(TargetWaypoint);
+		return true;
+	}
+
+	return false;
+}
+
+FString AVirtualCameraPawnBase::SaveWaypoint()
+{
+	// Convert index to string with leading zeros
+	FString WaypointNum = LeftPadWithZeros(WaypointIndex, MinimumIntegralDigits);
+
+	// Another waypoint has been created
+	WaypointIndex++;
+	FVirtualCameraWaypoint::NextIndex++;
+
+	// Create the waypoint
+	FVirtualCameraWaypoint NewWaypoint;
+	NewWaypoint.DateCreated = FDateTime::UtcNow();
+	NewWaypoint.Name = "Waypoint-" + WaypointNum;
+	NewWaypoint.WaypointTransform = CineCamera->GetComponentTransform();
+	Waypoints.Add(NewWaypoint.Name, NewWaypoint);
+	return NewWaypoint.Name;
+}
+
+FString AVirtualCameraPawnBase::SavePreset(const bool bSaveCameraSettings, const bool bSaveStabilization, const bool bSaveAxisLocking, const bool bSaveMotionScale)
+{
+	// Convert index to string with leading zeros
+	FString PresetNum = LeftPadWithZeros(PresetIndex, MinimumIntegralDigits);
+	FString PresetName = "Preset-" + PresetNum;
+	
+	// Another preset has been created
+	PresetIndex++;
+	FVirtualCameraSettingsPreset::NextIndex++;
+
+	FVirtualCameraSettingsPreset PresetToAdd;
+	PresetToAdd.DateCreated = FDateTime::UtcNow();
+
+	PresetToAdd.bIsCameraSettingsSaved = bSaveCameraSettings;
+	PresetToAdd.bIsStabilizationSettingsSaved = bSaveStabilization;
+	PresetToAdd.bIsAxisLockingSettingsSaved = bSaveAxisLocking;
+	PresetToAdd.bIsMotionScaleSettingsSaved = bSaveMotionScale;
+
+	if (CineCamera)
+	{
+		PresetToAdd.CameraSettings.FocalLength = CineCamera->GetCurrentFocalLength();
+		PresetToAdd.CameraSettings.Aperture = CineCamera->GetCurrentAperture();
+		PresetToAdd.CameraSettings.FilmbackWidth = CineCamera->Filmback.SensorWidth;
+		PresetToAdd.CameraSettings.FilmbackHeight = CineCamera->Filmback.SensorHeight;	
+	}
+
+	if (MovementComponent)
+	{
+		PresetToAdd.CameraSettings.AxisSettings = MovementComponent->AxisSettings;
+	}
+
+	SettingsPresets.Add(PresetName, PresetToAdd);
+
+	SaveSettings();
+
+	return PresetName;
+}
+
+TMap<FString, FVirtualCameraSettingsPreset> AVirtualCameraPawnBase::GetSettingsPresets()
+{
+	SettingsPresets.KeySort([](const FString& a, const FString& b) -> bool
+	{
+		return a < b;
+	});
+
+	return SettingsPresets;
+}
+
+bool AVirtualCameraPawnBase::LoadPreset(const FString& PresetName)
+{
+	FVirtualCameraSettingsPreset* LoadedPreset = SettingsPresets.Find(PresetName);
+
+	if (LoadedPreset)
+	{
+		UpdateSettingsFromPreset(LoadedPreset);
+		return true;
+	}
+	
+	return false;
+}
+
+int32 AVirtualCameraPawnBase::DeletePreset(const FString& PresetName)
+{
+	return SettingsPresets.Remove(PresetName);
+}
+
+int32 AVirtualCameraPawnBase::DeleteScreenshot(const FString& ScreenshotName)
+{
+	return Screenshots.Remove(ScreenshotName);
+}
+
+int32 AVirtualCameraPawnBase::DeleteWaypoint(const FString& WaypointName)
+{
+	return Waypoints.Remove(WaypointName);
+}
+
+void AVirtualCameraPawnBase::UpdateSettingsFromPreset(FVirtualCameraSettingsPreset* PresetToLoad)
+{
+	if (MovementComponent) 
+	{
+		// Load all settings for all axes
+		for (TPair<EVirtualCameraAxis, FVirtualCameraAxisSettings>& AxisSettingPair : MovementComponent->AxisSettings)
+		{
+			// Locking Settings
+			if (PresetToLoad->bIsAxisLockingSettingsSaved)
+			{
+				AxisSettingPair.Value.bIsLocked = PresetToLoad->CameraSettings.AxisSettings[AxisSettingPair.Key].bIsLocked;	
+			}
+			
+			// Movement Scaling
+			if (PresetToLoad->bIsMotionScaleSettingsSaved)
+			{
+				AxisSettingPair.Value.MovementScale = PresetToLoad->CameraSettings.AxisSettings[AxisSettingPair.Key].MovementScale;
+			}
+			
+			// Stabilization Scaling
+			if (PresetToLoad->bIsStabilizationSettingsSaved)
+			{
+				AxisSettingPair.Value.StabilizationScale = PresetToLoad->CameraSettings.AxisSettings[AxisSettingPair.Key].StabilizationScale;
+			}
+		}
+	}
+
+	// Camera Settings
+	if (CineCamera)
+	{
+		if (PresetToLoad->bIsCameraSettingsSaved)
+		{
+			CineCamera->CurrentAperture = PresetToLoad->CameraSettings.Aperture;
+			CineCamera->CurrentFocalLength = PresetToLoad->CameraSettings.FocalLength;
+			CineCamera->Filmback.SensorWidth = PresetToLoad->CameraSettings.FilmbackWidth;
+			CineCamera->Filmback.SensorHeight = PresetToLoad->CameraSettings.FilmbackHeight;
+		}
+	}
+	
+	LoadFinished();
+}
+
+void AVirtualCameraPawnBase::SaveHomeWaypoint(const FString& NewHomeWaypointName)
+{
+	if (NewHomeWaypointName.IsEmpty() || Waypoints.Contains(NewHomeWaypointName))
+	{
+		// Remove home waypoint mark if one exists 
+		if (Waypoints.Contains(HomeWaypointName))
+		{
+			Waypoints[HomeWaypointName].bIsHomeWaypoint = false;
+		}
+		
+		HomeWaypointName = NewHomeWaypointName;
+
+		// If updating to new home waypoint name, set to true
+		if (!NewHomeWaypointName.IsEmpty()) 
+		{
+			Waypoints[HomeWaypointName].bIsHomeWaypoint = true;
+		}
+	}
+}
+
+bool AVirtualCameraPawnBase::TeleportToWaypoint(const FString& WaypointName)
+{
+	// Make sure waypoint index is valid
+	if (!Waypoints.Contains(WaypointName))
+	{
+		return false;
+	}
+
+	MovementComponent->Teleport(Waypoints[WaypointName].WaypointTransform);
+	return true;
+}
+
+bool AVirtualCameraPawnBase::TeleportToHomeWaypoint()
+{
+	return TeleportToWaypoint(HomeWaypointName);
+}
+
+FString AVirtualCameraPawnBase::TakeScreenshot()
+{
+	if (!CineCamera)
+	{
+		return FString();
+	}
+
+	// Convert index to string with leading zeros
+	FString ScreenshotNum = LeftPadWithZeros(ScreenshotIndex, MinimumIntegralDigits);
+	// Track that another screenshot has been added
+	ScreenshotIndex++;
+	FVirtualCameraScreenshot::NextIndex++;
+	
+	APlayerController* PC = Cast<APlayerController>(Controller);
+
+	if (PC)
+	{
+		UWorld* World = GetWorld();
+		// Apply aspect ratio restraints to image
+		if (CineCamera && World)
+		{
+			CineCamera->Filmback = CineCamera->DesiredFilmbackSettings;
+			CineCamera->StopCameraViewUpdates();
+
+			FScreenshotRequest::RequestScreenshot(false);
+		}
+	}
+
+	// Store all the data for this screenshot
+	FVirtualCameraScreenshot NewScreenshot;
+	NewScreenshot.Waypoint.DateCreated = FDateTime::UtcNow();
+	
+	// Screenshots are named with the "Screenshot-" prefix and their index.
+	// Name and transform are saved on the screenshot's internal waypoint.
+	NewScreenshot.Waypoint.Name = "Screenshot-" + ScreenshotNum;
+	NewScreenshot.Waypoint.WaypointTransform = CineCamera->GetComponentTransform();
+
+	NewScreenshot.CameraSettings.Aperture = CineCamera->GetCurrentAperture();
+	NewScreenshot.CameraSettings.FocalLength = CineCamera->GetCurrentFocalLength();
+
+	Screenshots.Add(NewScreenshot.Waypoint.Name, NewScreenshot);
+	return NewScreenshot.Waypoint.Name;
+}
+
+bool AVirtualCameraPawnBase::LoadScreenshotView(const FString& ScreenshotName)
+{
+	// Make sure waypoint index is valid
+	if (!CineCamera|| !Screenshots.Contains(ScreenshotName))
+	{
+		return false;
+	}
+
+	// Load the camera settings
+	FVirtualCameraScreenshot &ScreenshotToLoad = Screenshots[ScreenshotName];
+	int32 Index;
+	
+	// If the aperture is in the preset list, set it to that, otherwise do nothing
+	if (CineCamera->ApertureOptions.Find(ScreenshotToLoad.CameraSettings.Aperture, Index))
+	{
+		CineCamera->CurrentAperture = CineCamera->ApertureOptions[Index];
+	}
+	
+	if (CineCamera->FocalLengthOptions.Find(ScreenshotToLoad.CameraSettings.FocalLength, Index))
+	{
+		CineCamera->CurrentFocalLength = CineCamera->FocalLengthOptions[Index];
+	}
+	
+	// If locked on to an actor, break out of that lock
+	if (GetAttachParentActor())
+	{
+		DetachFromActor(FDetachmentTransformRules::KeepWorldTransform);
+	}
+
+	MovementComponent->Teleport(ScreenshotToLoad.Waypoint.WaypointTransform);
+	return true;
+}
+
+bool AVirtualCameraPawnBase::RenameScreenshotLocation(const FString& TargetScreenshot, const FString& NewScreenshotName)
+{
+	if (Screenshots.Contains(TargetScreenshot))
+	{
+		Screenshots.FindOrAdd(NewScreenshotName) = Screenshots[TargetScreenshot];
+		Screenshots[NewScreenshotName].Waypoint.Name = NewScreenshotName;
+		Screenshots.Remove(TargetScreenshot);
+		return true;
+	}
+
+	return false;
+}
+
+void AVirtualCameraPawnBase::SetDesiredDistanceUnits(const EUnit DesiredUnits)
+{
+	if (FUnitConversion::IsUnitOfType(DesiredUnits, EUnitType::Distance))
+	{
+		DesiredDistanceUnits = DesiredUnits;
+	}
+}
+
+void AVirtualCameraPawnBase::SetPresetFavoriteStatus(const FString& PresetName, const bool bIsFavorite)
+{
+	if (SettingsPresets.Contains(PresetName))
+	{
+		SettingsPresets[PresetName].bIsFavorited = bIsFavorite;
+	}
+}
+
+void AVirtualCameraPawnBase::SetScreenshotFavoriteStatus(const FString& ScreenshotName, const bool bIsFavorite)
+{
+	if (Screenshots.Contains(ScreenshotName))
+	{
+		Screenshots[ScreenshotName].Waypoint.bIsFavorited = bIsFavorite;
+	}
+}
+
+void AVirtualCameraPawnBase::SetWaypointFavoriteStatus(const FString& WaypointName, const bool bIsFavorite)
+{
+	if (Waypoints.Contains(WaypointName))
+	{
+		Waypoints[WaypointName].bIsFavorited = bIsFavorite;
+	}
+}
+
+void AVirtualCameraPawnBase::SaveSettings()
+{
+	if (!CineCamera)
+	{
+		return;
+	}
+
+	UVirtualCameraSaveGame* SaveGameInstance = Cast<UVirtualCameraSaveGame>(UGameplayStatics::CreateSaveGameObject(UVirtualCameraSaveGame::StaticClass()));
+
+	// Save waypoints
+	SaveGameInstance->Waypoints = Waypoints;
+	SaveGameInstance->HomeWaypointName = HomeWaypointName;
+
+	// Save screenshots
+	SaveGameInstance->Screenshots = Screenshots;
+
+	// Save focal length and aperture
+	SaveGameInstance->CameraSettings.FocalLength = CineCamera->GetCurrentFocalLength();
+	SaveGameInstance->CameraSettings.Aperture = CineCamera->GetCurrentAperture();
+	SaveGameInstance->CameraSettings.bAllowFocusVisualization = bAllowFocusVisualization;
+#if WITH_EDITORONLY_DATA
+	SaveGameInstance->CameraSettings.DebugFocusPlaneColor = CineCamera->FocusSettings.DebugFocusPlaneColor;
+#else
+	SaveGameInstance->CameraSettings.DebugFocusPlaneColor = FColor();
+#endif
+
+	// Save filmback settings
+	SaveGameInstance->CameraSettings.FilmbackName = CineCamera->GetCurrentFilmbackName();
+	SaveGameInstance->CameraSettings.FilmbackWidth = CineCamera->Filmback.SensorWidth;
+	SaveGameInstance->CameraSettings.FilmbackHeight = CineCamera->Filmback.SensorHeight;
+	SaveGameInstance->CameraSettings.MatteOpacity = CineCamera->MatteOpacity;
+
+	// Save axis settings
+	SaveGameInstance->CameraSettings.AxisSettings = MovementComponent->AxisSettings;
+
+	// Save settings presets
+	SaveGameInstance->SettingsPresets = SettingsPresets;
+
+	// Save indices for naming
+	SaveGameInstance->WaypointIndex = FVirtualCameraWaypoint::NextIndex;
+	SaveGameInstance->ScreenshotIndex = FVirtualCameraScreenshot::NextIndex;
+	SaveGameInstance->PresetIndex = FVirtualCameraSettingsPreset::NextIndex;
+
+	SaveGameInstance->CameraSettings.DesiredDistanceUnits = DesiredDistanceUnits;
+
+	// Write save file to disk
+	UGameplayStatics::SaveGameToSlot(SaveGameInstance, SavedSettingsSlotName, 0);
+}
+
+void AVirtualCameraPawnBase::LoadSettings()
+{
+	if (!CineCamera)
+	{
+		return;
+	}
+
+	UVirtualCameraSaveGame* SaveGameInstance = Cast<UVirtualCameraSaveGame>(UGameplayStatics::CreateSaveGameObject(UVirtualCameraSaveGame::StaticClass()));
+	SaveGameInstance = Cast<UVirtualCameraSaveGame>(UGameplayStatics::LoadGameFromSlot(SavedSettingsSlotName, 0));
+
+	if (!SaveGameInstance)
+	{
+		UE_LOG(LogActor, Warning, TEXT("Virtual camera pawn could not find save game to load, using default settings."))
+		return;
+	}
+	
+	// Load waypoints
+	Waypoints = SaveGameInstance->Waypoints;
+	HomeWaypointName = SaveGameInstance->HomeWaypointName;
+
+	FVirtualCameraWaypoint::NextIndex = SaveGameInstance->WaypointIndex;
+	if (Waypoints.Num() > FVirtualCameraWaypoint::NextIndex)
+	{
+		FVirtualCameraWaypoint::NextIndex = Waypoints.Num();
+	}
+
+	// Load screenshots
+	Screenshots = SaveGameInstance->Screenshots;
+
+	FVirtualCameraScreenshot::NextIndex = SaveGameInstance->ScreenshotIndex;
+	if (Screenshots.Num() > FVirtualCameraScreenshot::NextIndex)
+	{
+		FVirtualCameraScreenshot::NextIndex = Screenshots.Num();
+	}
+
+	bAllowFocusVisualization = SaveGameInstance->CameraSettings.bAllowFocusVisualization;
+
+#if WITH_EDITORONLY_DATA
+	if (SaveGameInstance->CameraSettings.DebugFocusPlaneColor != FColor())
+	{
+		CineCamera->FocusSettings.DebugFocusPlaneColor = SaveGameInstance->CameraSettings.DebugFocusPlaneColor;
+	}
+#endif
+
+	// Load focal length
+	if (CineCamera->FocalLengthOptions.Contains(SaveGameInstance->CameraSettings.FocalLength))
+	{
+		CineCamera->SetCurrentFocalLength(SaveGameInstance->CameraSettings.FocalLength);
+	}
+
+	// Load aperture
+	if (CineCamera->ApertureOptions.Contains(SaveGameInstance->CameraSettings.Aperture))
+	{
+		CineCamera->SetCurrentAperture(SaveGameInstance->CameraSettings.Aperture);
+	}
+
+	// Load filmback settings
+	if (!CineCamera->SetFilmbackPresetOption(SaveGameInstance->CameraSettings.FilmbackName))
+	{
+		// If name isn't found, use backup settings
+		CineCamera->Filmback.SensorWidth = SaveGameInstance->CameraSettings.FilmbackWidth;
+		CineCamera->Filmback.SensorHeight = SaveGameInstance->CameraSettings.FilmbackHeight;
+	}
+	CineCamera->MatteOpacity = SaveGameInstance->CameraSettings.MatteOpacity;
+
+	MovementComponent->AxisSettings = SaveGameInstance->CameraSettings.AxisSettings;
+	MovementComponent->ResetCameraOffsetsToTracker();
+
+	// load presets, but don't overwrite existing ones
+	SettingsPresets.Append(SaveGameInstance->SettingsPresets);
+
+	// If the saved preset index is smaller than total presets, set it so that it won't overwrite existing presets.
+	FVirtualCameraSettingsPreset::NextIndex = SaveGameInstance->PresetIndex;
+	if (SettingsPresets.Num() > FVirtualCameraSettingsPreset::NextIndex)
+	{
+		FVirtualCameraSettingsPreset::NextIndex = SettingsPresets.Num();
+	}
+
+	// Load values of indices for naming
+	PresetIndex = FVirtualCameraSettingsPreset::NextIndex;
+	ScreenshotIndex = FVirtualCameraScreenshot::NextIndex;
+	WaypointIndex = FVirtualCameraWaypoint::NextIndex;
+
+	DesiredDistanceUnits = SaveGameInstance->CameraSettings.DesiredDistanceUnits;
+}
+
+FString AVirtualCameraPawnBase::LeftPadWithZeros(int32 InNumber, int32 MinNumberOfCharacters) const
+{
+	FString ReturnString = FString::FromInt(InNumber);
+
+	while (ReturnString.Len() < MinNumberOfCharacters)
+	{
+		ReturnString = FString("0") + ReturnString;
+	}
+
+	return ReturnString;
+}
+
+//~ Begin IVirtualCameraPresetContainer Interface
+
+FString AVirtualCameraPawnBase::SavePreset_Implementation(const bool bSaveCameraSettings, const bool bSaveStabilization, const bool bSaveAxisLocking, const bool bSaveMotionScale)
+{
+	return SavePreset(bSaveCameraSettings, bSaveStabilization, bSaveAxisLocking, bSaveMotionScale);
+}
+
+bool AVirtualCameraPawnBase::LoadPreset_Implementation(const FString& PresetName)
+{
+	return LoadPreset(PresetName);
+}
+
+int32 AVirtualCameraPawnBase::DeletePreset_Implementation(const FString& PresetName)
+{
+	return DeletePreset(PresetName);
+}
+
+TMap<FString, FVirtualCameraSettingsPreset> AVirtualCameraPawnBase::GetSettingsPresets_Implementation()
+{
+	return GetSettingsPresets();
+}
+
+//~ End IVirtualCameraPresetContainer Interface
+
+//~ Begin IVirtualCameraOptions Interface
+
+void AVirtualCameraPawnBase::SetDesiredDistanceUnits_Implementation(const EUnit DesiredUnits)
+{
+	SetDesiredDistanceUnits(DesiredUnits);
+}
+
+EUnit AVirtualCameraPawnBase::GetDesiredDistanceUnits_Implementation()
+{
+	return GetDesiredDistanceUnits();
+}
+
+bool AVirtualCameraPawnBase::IsFocusVisualizationAllowed_Implementation()
+{
+	return IsFocusVisualizationAllowed();
+}
+
+//~ End IVirtualCameraOptions Interface
