Move VirtualCamera and VirtualCameraEditor to VirtualCamera plugin

diff --git a/Engine/Plugins/Experimental/VirtualCamera/Source/VirtualCamera/Private/VirtualCameraPlayerControllerBase.cpp b/Engine/Plugins/Experimental/VirtualCamera/Source/VirtualCamera/Private/VirtualCameraPlayerControllerBase.cpp
new file mode 100644
--- /dev/null
+++ b/Engine/Plugins/Experimental/VirtualCamera/Source/VirtualCamera/Private/VirtualCameraPlayerControllerBase.cpp
@@ -0,0 +1,1423 @@
+// Copyright Epic Games, Inc. All Rights Reserved.
+
+#include "VirtualCameraPlayerControllerBase.h"
+
+#include "Actors/VPRootActor.h"
+#include "Components/InputComponent.h"
+#include "ConcertVirtualCamera.h"
+#include "CineCameraActor.h"
+#include "Engine/Engine.h"
+#include "Engine/World.h"
+#include "Engine/GameViewportClient.h"
+#include "Features/IModularFeatures.h"
+#include "Framework/Application/SlateApplication.h"
+#include "ILiveLinkClient.h"
+#include "IXRTrackingSystem.h"
+#include "RemoteSession.h"
+#include "Channels/RemoteSessionFrameBufferChannel.h"
+#include "Channels/RemoteSessionXRTrackingChannel.h"
+#include "Roles/LiveLinkAnimationRole.h"
+#include "Roles/LiveLinkAnimationTypes.h"
+#include "Roles/LiveLinkTransformRole.h"
+#include "Roles/LiveLinkTransformTypes.h"
+#include "VirtualCamera.h"
+#include "VPGameMode.h"
+
+#if WITH_EDITOR
+#include "Recorder/TakeRecorderBlueprintLibrary.h"
+#endif
+
+
+AVirtualCameraPlayerControllerBase::AVirtualCameraPlayerControllerBase(const FObjectInitializer& ObjectInitializer)
+{
+	// Default tracker input source
+	InputSource = ETrackerInputSource::ARKit;
+
+	TargetCameraActorClass = ACineCameraActor::StaticClass();
+	TargetCameraActor = nullptr;
+
+	LiveLinkClient = nullptr;
+	IModularFeatures& ModularFeatures = IModularFeatures::Get();
+	if (ModularFeatures.IsModularFeatureAvailable(ILiveLinkClient::ModularFeatureName))
+	{
+		LiveLinkClient = &IModularFeatures::Get().GetModularFeature<ILiveLinkClient>(ILiveLinkClient::ModularFeatureName);
+	}
+
+	// Default autofocus to center of screen
+	AutoFocusScreenPosition = FVector2D(.5f, .5f);
+
+	// Default touch input values
+	TouchInputState = ETouchInputState::BlueprintDefined;
+	PreviousTouchInput = TouchInputState;
+	CurrentFocusMethod = EVirtualCameraFocusMethod::Manual;
+
+	// Cached values
+	bCachedShouldUpdateTargetCameraTransform = true;
+	bCachedIsVirtualCameraControlledByRemoteSession = true;
+}
+
+void AVirtualCameraPlayerControllerBase::OnPossess(APawn* InPawn)
+{
+	bCachedShouldUpdateTargetCameraTransform = ShouldUpdateTargetCameraTransform();
+	bCachedIsVirtualCameraControlledByRemoteSession = IsVirtualCameraControlledByRemoteSession();
+
+	Super::OnPossess(InPawn);
+}
+
+void AVirtualCameraPlayerControllerBase::BeginPlay()
+{
+	// Find the Root Actor
+	RootActor = nullptr;
+	AVPGameMode* VPGameMode = Cast<AVPGameMode>(GetWorld()->GetAuthGameMode());
+	if (VPGameMode)
+	{
+		RootActor = VPGameMode->GetRootActor();
+	}
+	else
+	{
+		UE_LOG(LogVirtualCamera, Warning, TEXT("The Game Mode is not a VPGameMode."));
+	}
+
+	if (RootActor == nullptr)
+	{
+		UE_LOG(LogVirtualCamera, Warning, TEXT("There is no VP Root Actor in the scene. A CineCameraActor will be spawned. Multi user functionalities may suffer."));
+	}
+
+	// Make a default sequence playback controller
+	LevelSequencePlaybackController = NewObject<ULevelSequencePlaybackController>(this);
+	if (LevelSequencePlaybackController)
+	{
+		// Get the Cine Camera actor used by that root actor
+		if (RootActor)
+		{
+			TargetCameraActor = RootActor->GetCineCameraActor();
+			if (TargetCameraActor == nullptr)
+			{
+				UE_LOG(LogVirtualCamera, Warning, TEXT("The Root Actor doesn't have a cinematic camera."));
+			}
+		}
+
+		if (TargetCameraActor == nullptr)
+		{
+			// Spawn the target camera that the level sequence will use as a target
+			TargetCameraActor = GetWorld()->SpawnActor<ACineCameraActor>(TargetCameraActorClass);
+		}
+
+		// bLockToHmd is set to true by default. Remove it to prevent unwanted movement from XR system.
+		TargetCameraActor->GetCameraComponent()->bLockToHmd = false;
+
+		// Bind to record enabled state change delegate
+		LevelSequencePlaybackController->OnRecordEnabledStateChanged.BindUObject(this, &AVirtualCameraPlayerControllerBase::HandleRecordEnabledStateChange);
+	}
+
+	// Is this Controller a RemoteSession controller
+	if (bCachedIsVirtualCameraControlledByRemoteSession)
+	{
+		if (IRemoteSessionModule* RemoteSession = FModuleManager::LoadModulePtr<IRemoteSessionModule>("RemoteSession"))
+		{
+			//RemoteSession->AddSupportedChannel(FRemoteSessionImageChannel::StaticType(), ERemoteSessionChannelMode::Write);
+			//RemoteSession->AddSupportedChannel(FRemoteSessionInputChannel::StaticType(), ERemoteSessionChannelMode::Read);
+			//RemoteSession->AddSupportedChannel(FRemoteSessionXRTrackingChannel::StaticType(), ERemoteSessionChannelMode::Read);
+
+			RemoteSession->InitHost();
+		}
+	}
+	else
+	{
+		SetViewTargetWithBlend(TargetCameraActor);
+	}
+
+	if (UVirtualCameraCineCameraComponent* CineCamera = GetVirtualCameraCineCameraComponent())
+	{
+		// Need to make sure we don't let ARKit control camera completely
+		CineCamera->bLockToHmd = false;
+	}
+
+	if (UVirtualCameraMovementComponent* MovementComponent = GetVirtualCameraMovementComponent())
+	{
+		MovementComponent->OnOffsetReset.AddDynamic(this, &AVirtualCameraPlayerControllerBase::BroadcastOffsetReset);
+	}
+
+	// Initialize the view of the camera with offsets taken into account
+	UpdatePawnWithTrackerData();
+	ResetCameraOffsetsToTracker();
+
+	Super::BeginPlay();
+}
+
+void AVirtualCameraPlayerControllerBase::Tick(float DeltaSeconds)
+{
+	Super::Tick(DeltaSeconds);
+
+	AVirtualCameraPawnBase* VCPawn = GetVirtualCameraPawn();
+	if (VCPawn)
+	{
+		UpdatePawnWithTrackerData();
+
+		if (VCPawn->IsAutoFocusEnabled())
+		{
+			// Auto focus is just setting screen focus through the Auto Focus position.
+			SetFocusDistanceThroughPoint(AutoFocusScreenPosition);
+		}
+	}
+
+	UVirtualCameraCineCameraComponent* VCCamera = GetVirtualCameraCineCameraComponent();
+	if (VCCamera)
+	{
+		if (VCCamera->GetCurrentFocusMethod() == EVirtualCameraFocusMethod::Auto)
+		{
+			SetFocusDistanceThroughPoint(AutoFocusScreenPosition);
+		}
+
+		VCCamera->UpdateCameraView();
+	}
+
+	if (LevelSequencePlaybackController)
+	{
+		PilotTargetedCamera(VCPawn, VCCamera);
+	}
+}
+
+void AVirtualCameraPlayerControllerBase::SetupInputComponent()
+{
+	Super::SetupInputComponent();
+	InputComponent->BindTouch(IE_Released, this, &AVirtualCameraPlayerControllerBase::OnTouchInput);
+	InputComponent->BindAxisKey(EKeys::Gamepad_LeftY, this, &AVirtualCameraPlayerControllerBase::OnMoveForward);
+	InputComponent->BindAxisKey(EKeys::Gamepad_LeftX, this, &AVirtualCameraPlayerControllerBase::OnMoveRight);
+	InputComponent->BindAxisKey(EKeys::Gamepad_RightY, this, &AVirtualCameraPlayerControllerBase::OnMoveUp);
+}
+
+void AVirtualCameraPlayerControllerBase::InitializeAutoFocusPoint()
+{
+	if (GEngine && GEngine->GameViewport)
+	{
+		GEngine->GameViewport->GetViewportSize(AutoFocusScreenPosition);
+		AutoFocusScreenPosition.X *= 0.5f;
+		AutoFocusScreenPosition.Y *= 0.5f;
+	}
+	UpdateFocusReticle(FVector(AutoFocusScreenPosition, 0.f));
+}
+
+ACineCameraActor* AVirtualCameraPlayerControllerBase::GetTargetCamera()
+{
+	return TargetCameraActor;
+}
+
+void AVirtualCameraPlayerControllerBase::SetInputSource(ETrackerInputSource InInputSource)
+{
+	if (InputSource != InInputSource)
+	{
+		InputSource = InInputSource;
+
+		bCachedShouldUpdateTargetCameraTransform = ShouldUpdateTargetCameraTransform();
+	}
+}
+
+FString AVirtualCameraPlayerControllerBase::GetDistanceInDesiredUnits(const float InputDistance, const EUnit ConversionUnit) const
+{
+	// Checks if the specified conversion unit is a unit of distance, since this function assumes conversion from Unreal Units
+	if (!FUnitConversion::IsUnitOfType(ConversionUnit, EUnitType::Distance))
+	{
+		UE_LOG(LogVirtualCamera, Warning, TEXT("Conversion unit selected is not a unit of distance."))
+		return FString();
+	}
+
+	// ToDo: Add support for different project settings if default is changed
+	float ConvertedDistance = FUnitConversion::Convert<float>(InputDistance, EUnit::Centimeters, ConversionUnit);
+
+	FString ReturnString;
+
+	if (ConversionUnit == EUnit::Feet)
+	{
+		int32 Feet = FMath::FloorToInt(ConvertedDistance);
+		int32 Inches = FMath::RoundToInt(FMath::Frac(ConvertedDistance) * 12.f);
+
+		// Handle when inches rounds to 12
+		if (Inches == 12)
+		{
+			// Feet should increase by one instead of having 12 inches displayed
+			Feet++;
+			Inches = 0;
+		}
+
+		// Use ' and " for feet an inches unit labels
+		ReturnString = FString::FromInt(Feet);
+		ReturnString += "'";
+		
+		// Only display inches if there is enough room
+		if (Feet < 10000)
+		{
+			ReturnString += FString::FromInt(Inches);
+			ReturnString += "\"";
+		}
+		
+	}
+	else
+	{
+		// Meters
+		FNumberFormattingOptions NumberFormat;
+		// Only show one degree of precision for decimal values
+		NumberFormat.MinimumFractionalDigits = 0;
+		NumberFormat.MaximumFractionalDigits = 1;
+		ReturnString = FText::AsNumber(ConvertedDistance, &NumberFormat).ToString();
+		if (ConvertedDistance >= 1000) 
+		{
+			// Don't show decimal place if not enough room
+			ReturnString = FString::FromInt(FMath::RoundToInt(ConvertedDistance));
+		}
+		ReturnString += FUnitConversion::GetUnitDisplayString(ConversionUnit);
+	}
+
+	return ReturnString;
+}
+
+static const FName RemoteSessionTrackingSystemName(TEXT("RemoteSessionXRTrackingProxy"));
+
+bool AVirtualCameraPlayerControllerBase::GetCurrentTrackerLocationAndRotation(FVector& OutTrackerLocation, FRotator& OutTrackerRotation)
+{
+	bool bTransformSet = false;
+
+	switch (InputSource)
+	{
+		case ETrackerInputSource::ARKit:
+			if (GEngine && GEngine->XRSystem.IsValid() && GEngine->XRSystem->GetSystemName() == RemoteSessionTrackingSystemName)
+			{
+				FQuat ARKitQuaternion;
+				GEngine->XRSystem->GetCurrentPose(0, ARKitQuaternion, OutTrackerLocation);
+				OutTrackerRotation = ARKitQuaternion.Rotator();
+				bTransformSet = true;
+			}
+			break;
+
+		case ETrackerInputSource::LiveLink:
+			if (LiveLinkClient)
+			{
+				FLiveLinkSubjectFrameData EvaluateData;
+				if (LiveLinkClient->EvaluateFrame_AnyThread(LiveLinkTargetName, ULiveLinkTransformRole::StaticClass(), EvaluateData))
+				{
+					FLiveLinkTransformFrameData* TransformFrameData = EvaluateData.FrameData.Cast<FLiveLinkTransformFrameData>();
+					check(TransformFrameData);
+
+					OutTrackerLocation = TransformFrameData->Transform.GetLocation();
+					OutTrackerRotation = TransformFrameData->Transform.GetRotation().Rotator();
+
+					bTransformSet = true;
+				}
+				else if (LiveLinkClient->EvaluateFrame_AnyThread(LiveLinkTargetName, ULiveLinkAnimationRole::StaticClass(), EvaluateData))
+				{
+					FLiveLinkAnimationFrameData* AnimationFrameData = EvaluateData.FrameData.Cast<FLiveLinkAnimationFrameData>();
+					FLiveLinkSkeletonStaticData* AnimationStaticData = EvaluateData.StaticData.Cast<FLiveLinkSkeletonStaticData>();
+					check(AnimationFrameData && AnimationStaticData);
+
+					int32 RootIndex = AnimationStaticData->FindRootBone();
+					if (AnimationFrameData->Transforms.IsValidIndex(RootIndex))
+					{
+						
+						OutTrackerLocation = AnimationFrameData->Transforms[RootIndex].GetLocation();
+						OutTrackerRotation = AnimationFrameData->Transforms[RootIndex].GetRotation().Rotator();
+
+						bTransformSet = true;
+					}
+				}
+			}
+			break;
+
+		case ETrackerInputSource::Custom:
+			GetCustomTrackerLocationAndRotation(OutTrackerLocation, OutTrackerRotation);
+			bTransformSet = true;
+			break;
+
+		default:
+			UE_LOG(LogVirtualCamera, Warning, TEXT("Selected tracker source is not yet supported"))
+			break;
+	}
+
+	bTransformSet = bTransformSet && !OutTrackerLocation.ContainsNaN() && !OutTrackerRotation.ContainsNaN();
+
+	// Return failure status if we couldn't find device to track or device isn't supported or the values contains NaN
+	return bTransformSet;
+}
+
+bool AVirtualCameraPlayerControllerBase::IsTouchInputInFocusMode()
+{
+	return TouchInputState == ETouchInputState::ActorFocusTargeting || TouchInputState == ETouchInputState::AutoFocusTargeting || TouchInputState == ETouchInputState::ManualTouchFocus;
+}
+
+void AVirtualCameraPlayerControllerBase::SetFocusDistanceToActor(const ETouchIndex::Type TouchIndex, const FVector& Location)
+{
+	AVirtualCameraPawnBase* VCPawn = GetVirtualCameraPawn();
+
+	// Don't try to execute if not currently possessing a virtual Camera pawn
+	if (!VCPawn)
+	{
+		return;
+	}
+
+	FVector TraceDirection;
+	FVector CameraWorldLocation;
+
+	if (!DeprojectScreenPositionToWorld(Location.X, Location.Y, CameraWorldLocation, TraceDirection))
+	{
+		// If projection fails, return now
+		return;
+	}
+
+	float FocusTraceDist = 1000000.f;
+
+	// Trace to get depth under auto focus position
+	FCollisionQueryParams TraceParams(SCENE_QUERY_STAT(UpdateAutoFocus), true);
+	FHitResult Hit;
+
+	FVector const TraceEnd = CameraWorldLocation + TraceDirection * FocusTraceDist;
+	bool const bHit = GetWorld()->LineTraceSingleByChannel(Hit, CameraWorldLocation, TraceEnd, ECC_Visibility, TraceParams);
+	
+	if (bHit)
+	{
+		VCPawn->SetFocusDistance(Hit.Distance);
+
+		// Set this actor as the tracked actor, and move the exact focus tracked point to where the touch occurred
+		if (AActor* HitActor = Hit.GetHitObjectHandle().FetchActor())
+		{
+			FVector TrackingPointOffset = HitActor->GetActorRotation().UnrotateVector(Hit.ImpactPoint - HitActor->GetActorLocation());
+			TrackingPointOffset /= HitActor->GetActorScale(); // Adjust for non-standard scales when we rotate the vector
+
+			VCPawn->SetTrackedActorForFocus(HitActor, TrackingPointOffset);
+
+#if WITH_EDITOR
+			bool bIsRecording = UTakeRecorderBlueprintLibrary::IsRecording();
+#else
+			bool bIsRecording = false;
+#endif
+			if (!bIsRecording)
+			{	
+				VCPawn->TriggerFocusPlaneTimer();
+				VCPawn->HighlightTappedActor(HitActor);
+			}
+		}
+	}
+}
+
+AVirtualCameraPawnBase* AVirtualCameraPlayerControllerBase::GetVirtualCameraPawn() const 
+{
+	return Cast<AVirtualCameraPawnBase>(GetPawn());
+}
+
+void AVirtualCameraPlayerControllerBase::HandleRecordEnabledStateChange(const bool bIsRecordEnabled)
+{
+	OnRecordEnabledStateChanged(bIsRecordEnabled);
+}
+
+void AVirtualCameraPlayerControllerBase::SetFocusDistanceThroughPoint(const FVector2D ScreenPosition)
+{
+	AVirtualCameraPawnBase* VCPawn = GetVirtualCameraPawn();
+
+	FVector TraceDirection;
+	FVector CameraWorldLocation;
+
+	if (!DeprojectScreenPositionToWorld(ScreenPosition.X, ScreenPosition.Y, CameraWorldLocation, TraceDirection))
+	{
+		// If projection fails, return now
+		return;
+	}
+
+	float FocusTraceDist = 1000000.f;
+
+	// Trace to get depth under auto focus position
+	FCollisionQueryParams TraceParams(SCENE_QUERY_STAT(UpdateAutoFocus), true);
+	FHitResult Hit;
+
+	FVector const TraceEnd = CameraWorldLocation + TraceDirection * FocusTraceDist;
+	bool const bHit = GetWorld()->LineTraceSingleByChannel(Hit, CameraWorldLocation, TraceEnd, ECC_Visibility, TraceParams);
+
+	if (bHit && Hit.GetHitObjectHandle().IsValid())
+	{
+		VCPawn->SetFocusDistance(Hit.Distance);
+	}	
+}
+
+UVirtualCameraCineCameraComponent* AVirtualCameraPlayerControllerBase::GetVirtualCameraCineCameraComponent() const 
+{
+	if (AVirtualCameraPawnBase* VCPawn = GetVirtualCameraPawn())
+	{
+		return VCPawn->CineCamera;
+	}
+
+	return nullptr;
+}
+
+UVirtualCameraMovementComponent* AVirtualCameraPlayerControllerBase::GetVirtualCameraMovementComponent() const 
+{
+	if (AVirtualCameraPawnBase* VCPawn = GetVirtualCameraPawn())
+	{
+		return VCPawn->MovementComponent;
+	}
+
+	return nullptr;
+}
+
+void AVirtualCameraPlayerControllerBase::PilotTargetedCamera(AVirtualCameraPawnBase* PawnToFollow, UVirtualCameraCineCameraComponent* CameraToFollow)
+{
+	if (TargetCameraActor == nullptr || PawnToFollow == nullptr || CameraToFollow == nullptr)
+	{
+		return;
+	}
+
+	UCineCameraComponent* TargetCameraComponent = TargetCameraActor->GetCineCameraComponent();
+	if (TargetCameraComponent == nullptr)
+	{
+		return;
+	}
+
+	bool bAssignValues = false;
+	FConcertVirtualCameraControllerEvent CameraEvent;
+	if (bCachedIsVirtualCameraControlledByRemoteSession)
+	{
+		CameraEvent.InputSource = InputSource;
+		CameraEvent.CameraData = FConcertVirtualCameraCameraData(PawnToFollow, CameraToFollow);
+
+		IVirtualCameraModule::Get().GetConcertVirtualCameraManager()->SendControllerCameraEvent(CameraEvent);
+		bAssignValues = true;
+	}
+	else
+	{
+		bAssignValues = IVirtualCameraModule::Get().GetConcertVirtualCameraManager()->GetLatestControllerCameraEvent(CameraEvent);
+		if (bAssignValues)
+		{
+			CameraEvent.CameraData.ApplyTo(PawnToFollow, CameraToFollow, bCachedShouldUpdateTargetCameraTransform);
+
+			if (CameraEvent.InputSource != InputSource)
+			{
+				SetInputSource(CameraEvent.InputSource);
+			}
+		}
+	}
+
+	// Copy the info to the camera target camera
+	if (bAssignValues)
+	{
+		TargetCameraActor->SetActorLocationAndRotation(CameraEvent.CameraData.CameraActorLocation, CameraEvent.CameraData.CameraActorRotation);
+		TargetCameraComponent->CurrentAperture = CameraToFollow->CurrentAperture;
+		TargetCameraComponent->CurrentFocalLength = CameraToFollow->CurrentFocalLength;
+		TargetCameraComponent->FocusSettings = CameraToFollow->FocusSettings;
+		TargetCameraComponent->LensSettings = CameraToFollow->LensSettings;
+		TargetCameraComponent->Filmback = CameraToFollow->Filmback;
+	}
+	TargetCameraComponent->SetRelativeLocationAndRotation(CameraToFollow->GetRelativeLocation(), CameraToFollow->GetRelativeRotation());
+}
+
+bool AVirtualCameraPlayerControllerBase::InputTouch(uint32 Handle, ETouchType::Type Type, const FVector2D & TouchLocation, float Force, FDateTime DeviceTimestamp, uint32 TouchpadIndex)
+{
+	// Don't allow for input touch events if not within matte
+	if (!IsLocationWithinMatte(FVector(TouchLocation, 0.0f)))
+	{
+		return Super::InputTouch(Handle, Type, TouchLocation, Force, DeviceTimestamp, TouchpadIndex);;
+	}
+
+	// Allow touch and drag of the auto focus reticle when in auto focus targeting mode
+	if (Type == ETouchType::Moved && TouchInputState == ETouchInputState::AutoFocusTargeting)
+	{
+		UpdateScreenFocus(ETouchIndex::Touch1, FVector(TouchLocation, 0.f));
+	}
+
+	return Super::InputTouch(Handle, Type, TouchLocation, Force, DeviceTimestamp, TouchpadIndex);
+}
+
+void AVirtualCameraPlayerControllerBase::OnTouchInput(const ETouchIndex::Type TouchIndex, const FVector Location)
+{
+	if (!IsLocationWithinMatte(Location))
+	{
+		return;
+	}
+
+	switch (TouchInputState)
+	{
+		// Mode for attaching focus to an actor.
+		case ETouchInputState::ActorFocusTargeting:
+			SetFocusDistanceToActor(TouchIndex, Location);
+			UpdateScreenFocus(TouchIndex, Location);
+			break;
+
+		// Mode for changing the point on the screen used for auto focus targeting.
+		case ETouchInputState::AutoFocusTargeting:
+			UpdateScreenFocus(TouchIndex, Location);
+			ShowFocusPlaneFromTouch();
+			break;
+
+		case ETouchInputState::ManualTouchFocus:
+			SetFocusDistanceThroughPoint(FVector2D(Location.X, Location.Y));
+			UpdateScreenFocus(TouchIndex, Location);
+			ShowFocusPlaneFromTouch();
+			break;
+
+		// Allows for user defined behavior in blueprint
+		case ETouchInputState::BlueprintDefined:
+			break;
+
+		default:
+			break;
+	}
+	
+}
+
+void AVirtualCameraPlayerControllerBase::UpdateScreenFocus(const ETouchIndex::Type TouchIndex, const FVector Location)
+{
+	// Set New Screen Location for auto focus to trace through
+	AutoFocusScreenPosition.X = Location.X;
+	AutoFocusScreenPosition.Y = Location.Y;
+	UpdateFocusReticle(Location);
+}
+
+void AVirtualCameraPlayerControllerBase::OnMoveForward(const float InValue)
+{
+	if (UVirtualCameraMovementComponent* MovementComponent = GetVirtualCameraMovementComponent())
+	{
+		MovementComponent->OnMoveForward(InValue);
+	}
+}
+
+void AVirtualCameraPlayerControllerBase::OnMoveRight(const float InValue)
+{
+	if (UVirtualCameraMovementComponent* MovementComponent = GetVirtualCameraMovementComponent())
+	{
+		MovementComponent->OnMoveRight(InValue);
+	}
+}
+
+void AVirtualCameraPlayerControllerBase::OnMoveUp(const float InValue)
+{
+	if (UVirtualCameraMovementComponent* MovementComponent = GetVirtualCameraMovementComponent())
+	{
+		// ToDo: Figure out why this needs to be negated; maybe inverted look?
+		MovementComponent->OnMoveUp(-InValue);
+	}
+}
+
+void AVirtualCameraPlayerControllerBase::ShowFocusPlaneFromTouch()
+{
+	AVirtualCameraPawnBase* VCPawn = GetVirtualCameraPawn();
+
+	if (!VCPawn)
+	{
+		return;
+	}
+
+#if WITH_EDITOR
+	bool bIsRecording = UTakeRecorderBlueprintLibrary::IsRecording();
+#else
+	bool bIsRecording = false;
+#endif
+	if (!bIsRecording)
+	{
+		VCPawn->TriggerFocusPlaneTimer();
+	}
+}
+
+void AVirtualCameraPlayerControllerBase::UpdatePawnWithTrackerData()
+{
+	// Initialize the virtual camera view
+	if (AVirtualCameraPawnBase* VCPawn = GetVirtualCameraPawn())
+	{
+		FVector TrackerLocation = FVector::ZeroVector;
+		FRotator TrackerRotation = FRotator::ZeroRotator;
+		if (GetCurrentTrackerLocationAndRotation(TrackerLocation, TrackerRotation))
+		{
+			// Apply tracker offset to tracker; convert everything to transforms to make sure motions are calculated in the right order
+			FTransform TrackerRaw = FTransform(TrackerRotation, TrackerLocation);
+			FTransform AdjustedTracker = TrackerPostOffset.AsTransform() * TrackerRaw * TrackerPreOffset.AsTransform();
+			TrackerRotation = AdjustedTracker.Rotator();
+			TrackerLocation = AdjustedTracker.GetLocation();
+
+			VCPawn->ProcessMovementInput(TrackerLocation, TrackerRotation);
+		}
+	}
+}
+
+bool AVirtualCameraPlayerControllerBase::IsLocationWithinMatte(const FVector Location) const 
+{
+	if (UVirtualCameraCineCameraComponent* CineCamera = GetVirtualCameraCineCameraComponent())
+	{
+		FVector2D ViewportSize; 
+		GEngine->GameViewport->GetViewportSize(ViewportSize);
+		FVector2D MatteResolution;
+		FVector2D LowerBound;
+		FVector2D UpperBound;
+		
+		// Calculate the matte resolution based on viewport and ratio
+		MatteResolution.X = ViewportSize.X * CineCamera->ViewSizeRatio.X;
+		MatteResolution.Y = ViewportSize.Y * CineCamera->ViewSizeRatio.Y;
+
+		// Lower bound is half the difference between the two resolutions
+		LowerBound.X = (ViewportSize.X - MatteResolution.X) / 2;
+		LowerBound.Y = (ViewportSize.Y - MatteResolution.Y) / 2;
+
+		// Upper bound is the viewport size minus the lower bound
+		UpperBound.X = ViewportSize.X - LowerBound.X;
+		UpperBound.Y = ViewportSize.Y - LowerBound.Y;
+
+		// The touch location needs to be within the bounds
+		if (Location.X >= LowerBound.X && Location.X <= UpperBound.X && Location.Y >= LowerBound.Y && Location.Y <= UpperBound.Y)
+		{
+			return true;
+		}
+
+	}
+	
+	return false;
+}
+
+bool AVirtualCameraPlayerControllerBase::IsVirtualCameraControlledByRemoteSession_Implementation() const
+{
+	return true;
+}
+
+bool AVirtualCameraPlayerControllerBase::ShouldUpdateTargetCameraTransform_Implementation() const
+{
+	return InputSource == ETrackerInputSource::ARKit || IsVirtualCameraControlledByRemoteSession();
+}
+
+/***** UI Interface Functions *****/
+
+void AVirtualCameraPlayerControllerBase::AddBlendableToCamera(TScriptInterface<IBlendableInterface> BlendableToAdd, float Weight)
+{
+	if (UVirtualCameraCineCameraComponent* CineCamera = GetVirtualCameraCineCameraComponent())
+	{
+		CineCamera->AddBlendableToCamera(BlendableToAdd, Weight);
+	}
+}
+
+float AVirtualCameraPlayerControllerBase::ChangeAperturePreset(const bool bShiftUp)
+{
+	if (UVirtualCameraCineCameraComponent* CineCamera = GetVirtualCameraCineCameraComponent())
+	{
+		return CineCamera->ChangeAperturePreset(bShiftUp);
+	}
+
+	return 0.0f;
+}
+
+float AVirtualCameraPlayerControllerBase::ChangeFocalLengthPreset(const bool bShiftUp)
+{
+	if (UVirtualCameraCineCameraComponent* CineCamera = GetVirtualCameraCineCameraComponent())
+	{
+		return CineCamera->ChangeFocalLengthPreset(bShiftUp);
+	}
+
+	return 0.0f;
+}
+
+void AVirtualCameraPlayerControllerBase::ClearActiveLevelSequence()
+{
+	if (LevelSequencePlaybackController)
+	{
+		LevelSequencePlaybackController->ClearActiveLevelSequence();
+	}
+}
+
+int32 AVirtualCameraPlayerControllerBase::DeletePreset(FString PresetName)
+{
+	if (AVirtualCameraPawnBase* VCPawn = GetVirtualCameraPawn())
+	{
+		return VCPawn->DeletePreset(PresetName);
+	}
+	return -1;
+}
+
+int32 AVirtualCameraPlayerControllerBase::DeleteScreenshot(FString ScreenshotName)
+{
+	if (AVirtualCameraPawnBase* VCPawn = GetVirtualCameraPawn())
+	{
+		return VCPawn->DeleteScreenshot(ScreenshotName);
+	}
+	return -1;
+}
+
+int32 AVirtualCameraPlayerControllerBase::DeleteWaypoint(FString WaypointName)
+{
+	if (AVirtualCameraPawnBase* VCPawn = GetVirtualCameraPawn())
+	{
+		return VCPawn->DeleteWaypoint(WaypointName);
+	}
+	return -1;
+}
+
+FString AVirtualCameraPlayerControllerBase::GetActiveLevelSequenceName()
+{
+	if (LevelSequencePlaybackController)
+	{
+		return LevelSequencePlaybackController->GetActiveLevelSequenceName();
+	}
+
+	return FString();
+}
+
+ULevelSequence* AVirtualCameraPlayerControllerBase::GetActiveLevelSequence()
+{
+	if (LevelSequencePlaybackController)
+	{
+		return LevelSequencePlaybackController->GetActiveLevelSequence();
+	}
+
+	return nullptr;
+}
+
+float AVirtualCameraPlayerControllerBase::GetAxisStabilizationScale(EVirtualCameraAxis AxisToRetrieve)
+{
+	if (UVirtualCameraMovementComponent* MovementComponent = GetVirtualCameraMovementComponent())
+	{
+		return MovementComponent->GetAxisStabilizationScale(AxisToRetrieve);
+	}
+
+	return 0.0f;
+}
+
+float AVirtualCameraPlayerControllerBase::GetAxisMovementScale(EVirtualCameraAxis AxisToRetrieve)
+{
+	if (UVirtualCameraMovementComponent* MovementComponent = GetVirtualCameraMovementComponent())
+	{
+		return MovementComponent->GetAxisMovementScale(AxisToRetrieve);
+	}
+
+	return 0.0f;
+}
+
+float AVirtualCameraPlayerControllerBase::GetCurrentAperture()
+{
+	if (UVirtualCameraCineCameraComponent* CineCamera = GetVirtualCameraCineCameraComponent())
+	{
+		return CineCamera->GetCurrentAperture();
+	}
+
+	return 0.0f;
+}
+
+FString AVirtualCameraPlayerControllerBase::GetCurrentFilmbackName()
+{
+	if (UVirtualCameraCineCameraComponent* CineCamera = GetVirtualCameraCineCameraComponent())
+	{
+		return CineCamera->GetCurrentFilmbackName();
+	}
+
+	return FString();
+}
+
+float AVirtualCameraPlayerControllerBase::GetCurrentFocalLength()
+{
+	if (UVirtualCameraCineCameraComponent* CineCamera = GetVirtualCameraCineCameraComponent())
+	{
+		return CineCamera->GetCurrentFocalLength();
+	}
+
+	return 0.0f;
+}
+
+float AVirtualCameraPlayerControllerBase::GetCurrentFocusDistance()
+{
+	if (UVirtualCameraCineCameraComponent* CineCamera = GetVirtualCameraCineCameraComponent())
+	{
+		return CineCamera->GetCurrentFocusDistance();
+	}
+
+	return 0.0f;
+}
+
+FFrameNumber AVirtualCameraPlayerControllerBase::GetCurrentSequencePlaybackEnd()
+{
+	if (LevelSequencePlaybackController)
+	{
+		return LevelSequencePlaybackController->GetCurrentSequencePlaybackEnd();
+	}
+
+	return FFrameNumber();
+}
+
+FFrameNumber AVirtualCameraPlayerControllerBase::GetCurrentSequencePlaybackStart()
+{
+
+	if (LevelSequencePlaybackController)
+	{
+		return LevelSequencePlaybackController->GetCurrentSequencePlaybackStart();
+	}
+
+	return FFrameNumber();
+}
+
+bool AVirtualCameraPlayerControllerBase::IsSequencerLockedToCameraCut()
+{
+	if (LevelSequencePlaybackController)
+	{
+		return LevelSequencePlaybackController->IsSequencerLockedToCameraCut();
+	}
+
+	return false;
+}
+
+void AVirtualCameraPlayerControllerBase::SetSequencerLockedToCameraCut(bool bLockView)
+{
+	if (LevelSequencePlaybackController)
+	{
+		LevelSequencePlaybackController->SetSequencerLockedToCameraCut(bLockView);
+	}
+}
+
+FFrameRate AVirtualCameraPlayerControllerBase::GetCurrentSequenceFrameRate()
+{
+	if (LevelSequencePlaybackController)
+	{
+		return LevelSequencePlaybackController->GetCurrentSequenceFrameRate();
+	}
+
+	return FFrameRate();
+}
+
+EUnit AVirtualCameraPlayerControllerBase::GetDesiredDistanceUnits()
+{
+	if (AVirtualCameraPawnBase* VCPawn = GetVirtualCameraPawn())
+	{
+		return VCPawn->GetDesiredDistanceUnits();
+	}
+	return EUnit();
+}
+
+
+
+FColor AVirtualCameraPlayerControllerBase::GetFocusPlaneColor()
+{
+#if WITH_EDITOR
+	if (UVirtualCameraCineCameraComponent* CineCamera = GetVirtualCameraCineCameraComponent())
+	{
+		return CineCamera->FocusSettings.DebugFocusPlaneColor;
+	}
+#endif
+	return FColor();
+}
+
+bool AVirtualCameraPlayerControllerBase::GetFilmbackPresetOptions(TArray<FString>& OutFilmbackPresetsArray)
+{
+	if (UVirtualCameraCineCameraComponent* CineCamera = GetVirtualCameraCineCameraComponent())
+	{
+		return CineCamera->GetFilmbackPresetOptions(OutFilmbackPresetsArray);
+	}
+
+	return false;
+}
+
+FFrameNumber AVirtualCameraPlayerControllerBase::GetLevelSequenceLength()
+{
+	if (LevelSequencePlaybackController)
+	{
+		return LevelSequencePlaybackController->GetCurrentSequenceDuration();
+	}
+
+	return FFrameNumber();
+}
+
+void AVirtualCameraPlayerControllerBase::GetLevelSequences(TArray<FLevelSequenceData>& OutLevelSequenceNames)
+{
+	if (LevelSequencePlaybackController)
+	{
+		LevelSequencePlaybackController->GetLevelSequences(OutLevelSequenceNames);
+	}
+}
+
+float AVirtualCameraPlayerControllerBase::GetMatteAspectRatio()
+{
+	if (UVirtualCameraCineCameraComponent* CineCamera = GetVirtualCameraCineCameraComponent())
+	{
+		return CineCamera->GetMatteAspectRatio();
+	}
+	return 0.0f;
+}
+
+float AVirtualCameraPlayerControllerBase::GetMatteOpacity()
+{
+	if (UVirtualCameraCineCameraComponent* CineCamera = GetVirtualCameraCineCameraComponent())
+	{
+		return CineCamera->MatteOpacity;
+	}
+	return 0.0f;
+}
+
+void AVirtualCameraPlayerControllerBase::GetMatteValues(TArray<float>& OutMatteValues)
+{
+	if (UVirtualCameraCineCameraComponent* CineCamera = GetVirtualCameraCineCameraComponent())
+	{
+		CineCamera->GetMatteValues(OutMatteValues);
+		OutMatteValues.Sort();
+	}
+}
+
+FFrameTime AVirtualCameraPlayerControllerBase::GetPlaybackPosition()
+{
+	if (LevelSequencePlaybackController)
+	{
+		return LevelSequencePlaybackController->GetCurrentSequencePlaybackPosition();
+	}
+
+	return FFrameTime();
+}
+
+FTimecode AVirtualCameraPlayerControllerBase::GetPlaybackTimecode()
+{
+	if (LevelSequencePlaybackController)
+	{
+		return LevelSequencePlaybackController->GetCurrentSequencePlaybackTimecode();
+	}
+
+	return FTimecode();
+}
+
+void AVirtualCameraPlayerControllerBase::GetScreenshotInfo(FString ScreenshotName, FVirtualCameraScreenshot &OutScreenshotInfo)
+{
+	if (AVirtualCameraPawnBase* VCPawn = GetVirtualCameraPawn())
+	{
+		VCPawn->GetScreenshotInfo(ScreenshotName, OutScreenshotInfo);
+	}
+}
+
+void AVirtualCameraPlayerControllerBase::GetScreenshotNames(TArray<FString>& OutArray)
+{
+	if (AVirtualCameraPawnBase* VCPawn = GetVirtualCameraPawn())
+	{
+		VCPawn->GetScreenshotNames(OutArray);
+		// Sort OutArray reversed by name
+		OutArray.Sort([](const FString A, const FString B) {
+			return A > B;
+		});
+	}
+}
+
+TMap<FString, FVirtualCameraSettingsPreset> AVirtualCameraPlayerControllerBase::GetSettingsPresets()
+{
+	if (AVirtualCameraPawnBase* VCPawn = GetVirtualCameraPawn())
+	{
+		return VCPawn->GetSettingsPresets();
+	}
+
+	return TMap<FString, FVirtualCameraSettingsPreset>();
+}
+
+void AVirtualCameraPlayerControllerBase::GetWaypointInfo(FString WaypointName, FVirtualCameraWaypoint &OutWaypointInfo)
+{
+	if (AVirtualCameraPawnBase* VCPawn = GetVirtualCameraPawn())
+	{
+		VCPawn->GetWaypointInfo(WaypointName, OutWaypointInfo);
+	}
+}
+
+void AVirtualCameraPlayerControllerBase::GetWaypointNames(TArray<FString>& OutArray)
+{
+	if (AVirtualCameraPawnBase* VCPawn = GetVirtualCameraPawn())
+	{
+		VCPawn->GetWaypointNames(OutArray);
+		// Sort OutArray reversed by name
+		OutArray.Sort([] (const FString A, const FString B) {
+			return A > B;
+		});
+	}
+}
+
+bool AVirtualCameraPlayerControllerBase::IsAxisLocked(EVirtualCameraAxis AxisToCheck)
+{
+	if (UVirtualCameraMovementComponent* MovementComponent = GetVirtualCameraMovementComponent())
+	{
+		return MovementComponent->IsAxisLocked(AxisToCheck);
+	}
+	return false;
+}
+
+bool AVirtualCameraPlayerControllerBase::IsFocusVisualizationAllowed()
+{
+	if (AVirtualCameraPawnBase* VCPawn = GetVirtualCameraPawn())
+	{
+		return VCPawn->IsFocusVisualizationAllowed();
+	}
+
+	return false;
+}
+
+bool AVirtualCameraPlayerControllerBase::IsPlaying()
+{
+	if (LevelSequencePlaybackController)
+	{
+		return LevelSequencePlaybackController->IsSequencePlaybackActive();
+	}
+
+	return false;
+}
+
+bool AVirtualCameraPlayerControllerBase::IsUsingGlobalBoom()
+{
+	if (UVirtualCameraMovementComponent* VCMovementComponent = GetVirtualCameraMovementComponent())
+	{
+		return VCMovementComponent->IsUsingGlobalBoom();
+	}
+
+	return false;
+}
+
+void AVirtualCameraPlayerControllerBase::JumpToLevelSequenceEnd()
+{
+	if (LevelSequencePlaybackController)
+	{
+		LevelSequencePlaybackController->JumpToPlaybackPosition(LevelSequencePlaybackController->GetCurrentSequencePlaybackEnd());
+		LevelSequencePlaybackController->PauseLevelSequence();
+	}
+}
+
+void AVirtualCameraPlayerControllerBase::JumpToLevelSequenceStart()
+{
+	if (LevelSequencePlaybackController)
+	{
+		LevelSequencePlaybackController->JumpToPlaybackPosition(LevelSequencePlaybackController->GetCurrentSequencePlaybackStart());
+		LevelSequencePlaybackController->PauseLevelSequence();
+	}
+}
+
+void AVirtualCameraPlayerControllerBase::JumpToPlaybackPosition(const FFrameNumber& InFrameNumber)
+{
+	if (LevelSequencePlaybackController)
+	{
+		LevelSequencePlaybackController->JumpToPlaybackPosition(InFrameNumber);
+	}
+}
+
+bool AVirtualCameraPlayerControllerBase::LoadPreset(FString PresetName)
+{
+	if (AVirtualCameraPawnBase* VCPawn = GetVirtualCameraPawn())
+	{
+		return VCPawn->LoadPreset(PresetName);
+	}
+	return false;
+}
+
+bool AVirtualCameraPlayerControllerBase::LoadScreenshotView(const FString& ScreenshotName)
+{
+	if (AVirtualCameraPawnBase* VCPawn = GetVirtualCameraPawn())
+	{
+		VCPawn->LoadScreenshotView(ScreenshotName);
+	}
+
+	return false;
+}
+
+void AVirtualCameraPlayerControllerBase::PauseLevelSequence()
+{
+	if (LevelSequencePlaybackController)
+	{
+		LevelSequencePlaybackController->PauseLevelSequence();
+	}
+}
+
+void AVirtualCameraPlayerControllerBase::PlayLevelSequence()
+{
+	if (LevelSequencePlaybackController)
+	{
+		LevelSequencePlaybackController->PlayLevelSequence();
+	}
+}
+
+void AVirtualCameraPlayerControllerBase::PlayLevelSequenceInReverse()
+{
+	if (LevelSequencePlaybackController)
+	{
+		LevelSequencePlaybackController->PlayLevelSequenceReverse();
+	}
+}
+
+void AVirtualCameraPlayerControllerBase::ResetCameraOffsetsToTracker()
+{
+	if (UVirtualCameraMovementComponent* MovementComponent = GetVirtualCameraMovementComponent())
+	{
+		MovementComponent->ResetCameraOffsetsToTracker();
+	}
+}
+
+void AVirtualCameraPlayerControllerBase::ResumeLevelSequencePlay()
+{
+	if (LevelSequencePlaybackController)
+	{
+		LevelSequencePlaybackController->PlayLevelSequence();
+	}
+}
+
+void AVirtualCameraPlayerControllerBase::SaveHomeWaypoint(const FString& NewHomeWaypointName)
+{
+	if (AVirtualCameraPawnBase* VCPawn = GetVirtualCameraPawn())
+	{
+		VCPawn->SaveHomeWaypoint(NewHomeWaypointName);
+	}
+}
+
+FString AVirtualCameraPlayerControllerBase::SavePreset(bool bSaveCameraSettings, bool bSaveStabilization, bool bSaveAxisLocking, bool bSaveMotionScale)
+{
+	if (AVirtualCameraPawnBase* VCPawn = GetVirtualCameraPawn())
+	{
+		return VCPawn->SavePreset(bSaveCameraSettings, bSaveStabilization, bSaveAxisLocking, bSaveMotionScale);
+	}
+	return FString();
+}
+
+FString AVirtualCameraPlayerControllerBase::SaveWaypoint()
+{
+	if (AVirtualCameraPawnBase* VCPawn = GetVirtualCameraPawn())
+	{
+		return VCPawn->SaveWaypoint();
+	}
+
+	return FString();
+}
+
+bool AVirtualCameraPlayerControllerBase::SetActiveLevelSequence(ULevelSequence* InNewLevelSequence)
+{
+	if (LevelSequencePlaybackController)
+	{
+		return LevelSequencePlaybackController->SetActiveLevelSequence(InNewLevelSequence);
+	}
+
+	return false;
+}
+
+void AVirtualCameraPlayerControllerBase::SetAllowFocusPlaneVisualization(bool bShouldAllowFocusVisualization)
+{
+	if (AVirtualCameraPawnBase* VCPawn = GetVirtualCameraPawn())
+	{
+		VCPawn->SetAllowFocusPlaneVisualization(bShouldAllowFocusVisualization);
+	}
+}
+
+float AVirtualCameraPlayerControllerBase::SetAxisStabilizationScale(EVirtualCameraAxis AxisToAdjust, float NewStabilizationAmount)
+{
+	if (UVirtualCameraMovementComponent* MovementComponent = GetVirtualCameraMovementComponent())
+	{
+		return MovementComponent->SetAxisStabilizationScale(AxisToAdjust, NewStabilizationAmount);
+	}
+
+	return 0.0f;
+}
+
+void AVirtualCameraPlayerControllerBase::SetCurrentAperture(float NewAperture)
+{
+	if (UVirtualCameraCineCameraComponent* CineCamera = GetVirtualCameraCineCameraComponent())
+	{
+		CineCamera->SetCurrentAperture(NewAperture);
+	}
+}
+
+void AVirtualCameraPlayerControllerBase::SetCurrentFocalLength(const float NewFocalLength)
+{
+	if (UVirtualCameraCineCameraComponent* CineCamera = GetVirtualCameraCineCameraComponent())
+	{
+		CineCamera->SetCurrentFocalLength(NewFocalLength);
+	}
+}
+
+void AVirtualCameraPlayerControllerBase::SetCurrentFocusDistance(const float NewFocusDistance)
+{
+	if (UVirtualCameraCineCameraComponent* CineCamera = GetVirtualCameraCineCameraComponent())
+	{
+		CineCamera->SetCurrentFocusDistance(NewFocusDistance);
+	}
+}
+
+void AVirtualCameraPlayerControllerBase::SetDesiredDistanceUnits(const EUnit DesiredUnits)
+{
+	if (AVirtualCameraPawnBase* VCPawn = GetVirtualCameraPawn())
+	{
+		VCPawn->SetDesiredDistanceUnits(DesiredUnits);
+	}
+}
+
+bool AVirtualCameraPlayerControllerBase::SetFilmbackPresetOption(const FString & NewFilmbackPreset)
+{
+	if (UVirtualCameraCineCameraComponent* CineCamera = GetVirtualCameraCineCameraComponent())
+	{
+		return CineCamera->SetFilmbackPresetOption(NewFilmbackPreset);
+	}
+
+	return false;
+}
+
+void AVirtualCameraPlayerControllerBase::SetFocusMethod(const EVirtualCameraFocusMethod NewFocusMethod)
+{
+	if (UVirtualCameraCineCameraComponent* CineCamera = GetVirtualCameraCineCameraComponent())
+	{
+		CurrentFocusMethod = NewFocusMethod;
+		FocusMethodChanged(NewFocusMethod);
+		return CineCamera->SetFocusMethod(NewFocusMethod);
+	}
+}
+
+void AVirtualCameraPlayerControllerBase::SetFocusPlaneColor(const FColor NewFocusPlaneColor)
+{
+#if WITH_EDITOR
+	if (GetVirtualCameraCineCameraComponent())
+	{
+		GetVirtualCameraCineCameraComponent()->FocusSettings.DebugFocusPlaneColor = NewFocusPlaneColor;
+	}
+#endif
+}
+
+void AVirtualCameraPlayerControllerBase::SetFocusVisualization(bool bShowFocusVisualization)
+{
+	if (GetVirtualCameraCineCameraComponent())
+	{
+		GetVirtualCameraCineCameraComponent()->SetFocusVisualization(bShowFocusVisualization);
+	}
+}
+
+bool AVirtualCameraPlayerControllerBase::IsFocusVisualizationActivated() const
+{
+	if (GetVirtualCameraCineCameraComponent())
+	{
+		return GetVirtualCameraCineCameraComponent()->IsFocusVisualizationActivated();
+	}
+	return false;
+}
+
+bool AVirtualCameraPlayerControllerBase::SetMatteAspectRatio(const float NewMatteAspectRatio)
+{
+	if (UVirtualCameraCineCameraComponent* CineCamera = GetVirtualCameraCineCameraComponent())
+	{
+		return CineCamera->SetMatteAspectRatio(NewMatteAspectRatio);
+	}
+	return false;
+}
+
+void AVirtualCameraPlayerControllerBase::SetMatteOpacity(const float NewMatteOpacity)
+{
+	if (UVirtualCameraCineCameraComponent* CineCamera = GetVirtualCameraCineCameraComponent())
+	{
+		CineCamera->MatteOpacity = NewMatteOpacity;
+	}
+}
+
+void AVirtualCameraPlayerControllerBase::SetMovementScale(const EVirtualCameraAxis AxisToAdjust, const float NewMovementScale)
+{
+	if (UVirtualCameraMovementComponent* MovementComponent = GetVirtualCameraMovementComponent())
+	{
+		MovementComponent->SetMovementScale(AxisToAdjust, NewMovementScale);
+	}
+}
+
+void AVirtualCameraPlayerControllerBase::SetPresetFavoriteStatus(const FString& PresetName, const bool bIsFavorite)
+{
+	if (AVirtualCameraPawnBase* VCPawn = GetVirtualCameraPawn())
+	{
+		VCPawn->SetPresetFavoriteStatus(PresetName, bIsFavorite);
+	}
+}
+
+void AVirtualCameraPlayerControllerBase::SetSaveSettingsWhenClosing(bool bShouldSettingsSave)
+{
+	if (AVirtualCameraPawnBase* VCPawn = GetVirtualCameraPawn())
+	{
+		VCPawn->SetSaveSettingsWhenClosing(bShouldSettingsSave);
+	}
+}
+
+void AVirtualCameraPlayerControllerBase::SetScreenshotFavoriteStatus(const FString& ScreenshotName, const bool bIsFavorite)
+{
+	if (AVirtualCameraPawnBase* VCPawn = GetVirtualCameraPawn())
+	{
+		VCPawn->SetScreenshotFavoriteStatus(ScreenshotName, bIsFavorite);
+	}
+}
+
+void AVirtualCameraPlayerControllerBase::SetUseGlobalBoom(bool bShouldUseGlobalBoom)
+{
+	if (UVirtualCameraMovementComponent* VCMovementComponent = GetVirtualCameraMovementComponent())
+	{
+		VCMovementComponent->SetUseGlobalBoom(bShouldUseGlobalBoom);
+	}
+}
+
+void AVirtualCameraPlayerControllerBase::SetWaypointFavoriteStatus(const FString& WaypointName, const bool bIsFavorite)
+{
+	if (AVirtualCameraPawnBase* VCPawn = GetVirtualCameraPawn())
+	{
+		VCPawn->SetWaypointFavoriteStatus(WaypointName, bIsFavorite);
+	}
+}
+
+void AVirtualCameraPlayerControllerBase::SetZeroDutchOnLock(const bool bInValue)
+{
+	if (UVirtualCameraMovementComponent* MovementComponent = GetVirtualCameraMovementComponent())
+	{
+		MovementComponent->SetZeroDutchOnLock(bInValue);
+	}
+}
+
+bool AVirtualCameraPlayerControllerBase::ShouldSaveSettingsWhenClosing()
+{
+	if (AVirtualCameraPawnBase* VCPawn = GetVirtualCameraPawn())
+	{
+		return VCPawn->GetSaveSettingsWhenClosing();
+	}
+
+	return false;
+}
+
+void AVirtualCameraPlayerControllerBase::StopLevelSequencePlay()
+{
+	if (LevelSequencePlaybackController)
+	{
+		LevelSequencePlaybackController->StopLevelSequencePlay();
+	}
+}
+
+FString AVirtualCameraPlayerControllerBase::TakeScreenshot()
+{
+	if (AVirtualCameraPawnBase* VCPawn = GetVirtualCameraPawn())
+	{
+		return VCPawn->TakeScreenshot();
+	}
+
+	return FString();
+}
+
+void AVirtualCameraPlayerControllerBase::TeleportToHomeWaypoint()
+{
+	if (GetVirtualCameraPawn() && GetVirtualCameraPawn()->TeleportToHomeWaypoint())
+	{
+		// Do nothing since we teleport in the above call
+	}
+	else 
+	{
+		ResetCameraOffsetsToTracker();
+	}
+}
+
+bool AVirtualCameraPlayerControllerBase::TeleportToWaypoint(const FString& WaypointName)
+{
+	if (AVirtualCameraPawnBase* VCPawn = GetVirtualCameraPawn())
+	{
+		return VCPawn->TeleportToWaypoint(WaypointName);
+	}
+
+	return false;
+}
+
+bool AVirtualCameraPlayerControllerBase::ToggleAxisFreeze(const EVirtualCameraAxis AxisToToggle)
+{
+	if (UVirtualCameraMovementComponent* MovementComponent = GetVirtualCameraMovementComponent())
+	{
+		return MovementComponent->ToggleAxisFreeze(AxisToToggle);
+	}
+
+	return false;
+}
+
+bool AVirtualCameraPlayerControllerBase::ToggleAxisLock(const EVirtualCameraAxis AxisToToggle)
+{
+	if (UVirtualCameraMovementComponent* MovementComponent = GetVirtualCameraMovementComponent())
+	{
+		return MovementComponent->ToggleAxisLock(AxisToToggle);
+	}
+
+	return false;
+}
+
+void AVirtualCameraPlayerControllerBase::ActivateGameViewport()
+{
+	FSlateApplication::Get().ActivateGameViewport();
+}
diff --git a/Engine/Plugins/Experimental/VirtualCamera/Source/VirtualCamera/Private/VirtualCameraSaveGame.cpp b/Engine/Plugins/Experimental/VirtualCamera/Source/VirtualCamera/Private/VirtualCameraSaveGame.cpp
new file mode 100644
--- /dev/null
+++ b/Engine/Plugins/Experimental/VirtualCamera/Source/VirtualCamera/Private/VirtualCameraSaveGame.cpp
@@ -0,0 +1,14 @@
+// Copyright Epic Games, Inc. All Rights Reserved.
+
+#include "VirtualCameraSaveGame.h"
+
+int32 FVirtualCameraSettingsPreset::NextIndex = 1;
+int32 FVirtualCameraWaypoint::NextIndex = 1;
+int32 FVirtualCameraScreenshot::NextIndex = 1;
+
+UVirtualCameraSaveGame::UVirtualCameraSaveGame(const FObjectInitializer& ObjectInitializer)
+{
+	LLM_SCOPE_BYNAME("VirtualCamera/VirtualCameraSaveGame");
+	SaveSlotName = "SavedVirtualCameraSettings";
+	UserIndex = 0;
+}
\ No newline at end of file
diff --git a/Engine/Plugins/Experimental/VirtualCamera/Source/VirtualCamera/Private/VirtualCameraSubsystem.cpp b/Engine/Plugins/Experimental/VirtualCamera/Source/VirtualCamera/Private/VirtualCameraSubsystem.cpp
new file mode 100644
--- /dev/null
+++ b/Engine/Plugins/Experimental/VirtualCamera/Source/VirtualCamera/Private/VirtualCameraSubsystem.cpp
@@ -0,0 +1,102 @@
+// Copyright Epic Games, Inc. All Rights Reserved.
+
+#include "VirtualCameraSubsystem.h"
+
+#include "LevelSequence/LevelSequencePlaybackController.h"
+
+#include "Engine/Selection.h"
+#include "GameFramework/Actor.h"
+
+UVirtualCameraSubsystem::UVirtualCameraSubsystem()
+	: bIsStreaming(false)
+{
+	SequencePlaybackController = CreateDefaultSubobject<ULevelSequencePlaybackController>("SequencePlaybackController");
+#if WITH_EDITOR
+	USelection::SelectionChangedEvent.AddUObject(this, &UVirtualCameraSubsystem::HandleSelectionChangedEvent);
+	USelection::SelectObjectEvent.AddUObject(this, &UVirtualCameraSubsystem::HandleSelectObjectEvent);
+#endif
+}
+
+bool UVirtualCameraSubsystem::StartStreaming()
+{
+	if (bIsStreaming)
+	{
+		return false;
+	}
+
+	if (ActiveCameraController)
+	{
+		bIsStreaming = ActiveCameraController->StartStreaming();
+	}
+
+	if (bIsStreaming)
+	{
+		FEditorScriptExecutionGuard ScriptGuard;
+		OnStreamStartedDelegate.Broadcast();
+	}
+
+	return bIsStreaming;
+}
+
+bool UVirtualCameraSubsystem::StopStreaming()
+{
+	if (!bIsStreaming)
+	{
+		return false;
+	}
+
+	if (ActiveCameraController)
+	{
+		bIsStreaming = !(ActiveCameraController->StopStreaming());
+	}
+
+	if (!bIsStreaming)
+	{
+		FEditorScriptExecutionGuard ScriptGuard;
+		OnStreamStoppedDelegate.Broadcast();
+	}
+
+	return bIsStreaming;
+}
+
+bool UVirtualCameraSubsystem::IsStreaming() const
+{
+	return bIsStreaming;
+}
+
+#if WITH_EDITOR
+void UVirtualCameraSubsystem::HandleSelectionChangedEvent(UObject* ChangedObject)
+{
+	USelection* Selection = Cast<USelection>(ChangedObject);
+	if (Selection)
+	{
+		if (AActor* SelectedActor = Selection->GetBottom<AActor>()) 
+		{
+			OnSelectedAnyActorDelegate.Broadcast(SelectedActor); 
+		}
+	}
+}
+
+void UVirtualCameraSubsystem::HandleSelectObjectEvent(UObject* ChangedObject)
+{
+	if (AActor* SelectedActorCheck1 = Cast<AActor>(ChangedObject))
+	{
+		OnSelectedActorInViewportDelegate.Broadcast(SelectedActorCheck1);
+	}
+	else if (AActor* SelectedActorCheck2 = ChangedObject->GetTypedOuter<AActor>())
+	{
+		OnSelectedActorInViewportDelegate.Broadcast(SelectedActorCheck2);
+	}
+}
+#endif
+
+TScriptInterface<IVirtualCameraController> UVirtualCameraSubsystem::GetVirtualCameraController() const
+{
+	return ActiveCameraController;
+}
+
+void UVirtualCameraSubsystem::SetVirtualCameraController(TScriptInterface<IVirtualCameraController> VirtualCamera)
+{
+	ActiveCameraController = VirtualCamera;
+	//todo deactive the last current, initialize the new active, call back
+}
diff --git a/Engine/Plugins/Experimental/VirtualCamera/Source/VirtualCamera/Private/VirtualCameraUserSettings.cpp b/Engine/Plugins/Experimental/VirtualCamera/Source/VirtualCamera/Private/VirtualCameraUserSettings.cpp
new file mode 100644
--- /dev/null
+++ b/Engine/Plugins/Experimental/VirtualCamera/Source/VirtualCamera/Private/VirtualCameraUserSettings.cpp
@@ -0,0 +1,229 @@
+// Copyright Epic Games, Inc. All Rights Reserved.
+
+#include "VirtualCameraUserSettings.h"
+#include "UObject/SoftObjectPath.h"
+
+UVirtualCameraUserSettings::UVirtualCameraUserSettings()
+{
+
+	FSoftObjectPath DefaultVCamSoftPath = FSoftObjectPath(TEXT("/VirtualCamera/V2/VcamActor.VcamActor_C"));
+
+	if (DefaultVCamSoftPath.IsNull())
+	{
+		return;
+	}
+
+	DefaultVCamClass = TSoftClassPtr<AActor>(DefaultVCamSoftPath);
+
+}
+
+float UVirtualCameraUserSettings::GetFocusInterpSpeed()
+{
+	return FocusInterpSpeed;
+}
+
+void UVirtualCameraUserSettings::SetFocusInterpSpeed(const float InFocusInterpSpeed)
+{
+	FocusInterpSpeed = InFocusInterpSpeed;
+	SaveConfig();
+}
+
+float UVirtualCameraUserSettings::GetJoysticksSpeed()
+{
+	return JoysticksSpeed;
+}
+
+void UVirtualCameraUserSettings::SetJoysticksSpeed(const float InJoysticksSpeed)
+{
+	JoysticksSpeed = InJoysticksSpeed;
+	SaveConfig();
+}
+
+float UVirtualCameraUserSettings::GetMaxJoysticksSpeed()
+{
+    return MaxJoysticksSpeed;
+}
+
+void UVirtualCameraUserSettings::SetMaxJoysticksSpeed(const float InMaxJoysticksSpeed)
+{
+	MaxJoysticksSpeed = InMaxJoysticksSpeed;
+	SaveConfig();
+}
+
+bool UVirtualCameraUserSettings::IsMapGrayscle()
+{
+	return bIsMapGrayscale;
+}
+
+void UVirtualCameraUserSettings::SetIsMapGrayscle(const bool bInIsMapGrayscle)
+{
+	bIsMapGrayscale = bInIsMapGrayscle;
+	SaveConfig();
+}
+
+bool UVirtualCameraUserSettings::GetShouldOverrideCameraSettingsOnTeleport()
+{
+	return bOverrideCameraSettingsOnTeleportToScreenshot;
+}
+
+void UVirtualCameraUserSettings::SetShouldOverrideCameraSettingsOnTeleport(const bool bInOverrideCameraSettings)
+{
+	bOverrideCameraSettingsOnTeleportToScreenshot = bInOverrideCameraSettings;
+	SaveConfig();
+}
+
+FString UVirtualCameraUserSettings::GetSavedVirtualCameraFilmbackPresetName()
+{
+	return GetDefault<UVirtualCameraUserSettings>()->VirtualCameraFilmback;
+}
+
+void UVirtualCameraUserSettings::SetSavedVirtualCameraFilmbackPresetName(const FString& InFilmback)
+{
+	VirtualCameraFilmback = InFilmback;
+	SaveConfig();
+}
+
+bool UVirtualCameraUserSettings::GetShouldDisplayFilmLeader()
+{
+	return bDisplayFilmLeader;
+}
+
+void UVirtualCameraUserSettings::SetShouldDisplayFilmLeader(const bool bInDisplayFilmLeader)
+{
+	bDisplayFilmLeader = bInDisplayFilmLeader;
+	SaveConfig();
+}
+
+bool UVirtualCameraUserSettings::GetTeleportOnStart()
+{
+    return bTeleportOnStart;
+}
+
+void UVirtualCameraUserSettings::SetTeleportOnStart(const bool bInTeleportOnStart)
+{
+	bTeleportOnStart = bInTeleportOnStart;
+	SaveConfig();
+}
+
+
+
+void UVirtualCameraUserSettings::InjectGamepadKeybinds()
+{
+	ActionMappings = {
+		{TEXT("VirtualCamera_Home_EarSelection_Up"), EKeys::Gamepad_DPad_Up},
+		{TEXT("VirtualCamera_Home_EarSelection_Down"), EKeys::Gamepad_DPad_Down},
+		{TEXT("VirtualCamera_Home_EarSelection_Left"), EKeys::Gamepad_DPad_Left},
+		{TEXT("VirtualCamera_Home_EarSelection_Right"), EKeys::Gamepad_DPad_Right},
+		{TEXT("VirtualCamera_Home_Selection"), EKeys::Gamepad_FaceButton_Bottom},
+		{TEXT("VirtualCamera_Home_FStop"), EKeys::Gamepad_LeftTrigger},
+		{TEXT("VirtualCamera_Home_LensesMenu"), EKeys::Gamepad_RightTrigger},
+		{TEXT("VirtualCamera_Home_Screenshot"), EKeys::Gamepad_RightShoulder},
+		{TEXT("VirtualCamera_Home_ToggleStats"), EKeys::Gamepad_Special_Right},
+		{TEXT("VirtualCamera_Home_ToggleFlight"), EKeys::Gamepad_FaceButton_Top},
+		{TEXT("VirtualCamera_Home_ToggleRecord"), EKeys::Gamepad_LeftShoulder},
+		{TEXT("VirtualCamera_Playback_SelectsSwitch"), EKeys::Gamepad_LeftShoulder},
+		{TEXT("VirtualCamera_Playback_SelectClipsUp"), EKeys::Gamepad_DPad_Up},
+		{TEXT("VirtualCamera_Playback_SelectClipsDown"), EKeys::Gamepad_DPad_Down},
+		{TEXT("VirtualCamera_Playback_PlayPause"), EKeys::Gamepad_RightShoulder},
+		{TEXT("VirtualCamera_Playback_Fullscreen"), EKeys::Gamepad_FaceButton_Top},
+		{TEXT("VirtualCamera_Playback_MakeClipSelect"), EKeys::Gamepad_FaceButton_Left},
+		{TEXT("VirtualCamera_Playback_TargetClip"), EKeys::Gamepad_FaceButton_Bottom},
+		{TEXT("VirtualCamera_Playback_BackStart"), EKeys::Gamepad_DPad_Left},
+		{TEXT("VirtualCamera_Playback_Close"), EKeys::Gamepad_Special_Right},
+		{TEXT("VirtualCamera_Map_SwitchFilterSnapshotFlagLeft"), EKeys::Gamepad_DPad_Left},
+		{TEXT("VirtualCamera_Map_SwitchFilterSnapshotFlagRight"), EKeys::Gamepad_DPad_Right},
+		{TEXT("VirtualCamera_Map_Close"), EKeys::Gamepad_Special_Right},
+		{TEXT("VirtualCamera_Map_TeleportHome"), EKeys::Gamepad_FaceButton_Left},
+		{TEXT("VirtualCamera_Map_ToggleFullscreen"), EKeys::Gamepad_FaceButton_Top},
+		{TEXT("VirtualCamera_Recording_InitiateRecord"), EKeys::Gamepad_LeftTrigger},
+		{TEXT("VirtualCamera_Interaction_Undo"), EKeys::Gamepad_LeftShoulder},
+		{TEXT("VirtualCamera_Interaction_Redo"), EKeys::Gamepad_RightShoulder},
+		{TEXT("VirtualCamera_Interaction_SRTSwitchUp"), EKeys::Gamepad_DPad_Up},
+		{TEXT("VirtualCamera_Interaction_SRTSwitchDown"), EKeys::Gamepad_DPad_Down},
+		{TEXT("VirtualCamera_Interaction_XYZSwitchLeft"), EKeys::Gamepad_DPad_Left},
+		{TEXT("VirtualCamera_Interaction_XYZSwitchRight"), EKeys::Gamepad_DPad_Right},
+		{TEXT("VirtualCamera_Interaction_Close"), EKeys::Gamepad_Special_Right},
+		{TEXT("VirtualCamera_Interaction_LocalWorldToggle"), EKeys::Gamepad_FaceButton_Top},
+		{TEXT("VirtualCamera_Interaction_SelectSpawnNumpad"), EKeys::Gamepad_FaceButton_Bottom},
+		{TEXT("VirtualCamera_Focus_CloseMenu"), EKeys::Gamepad_Special_Right},
+		{TEXT("VirtualCamera_Focus_Selection"), EKeys::Gamepad_FaceButton_Bottom},
+		{TEXT("VirtualCamera_Focus_ToggleAuto"), EKeys::Gamepad_FaceButton_Top},
+		{TEXT("VirtualCamera_Playback_Back"), EKeys::Gamepad_FaceButton_Right},
+		{TEXT("VirtualCamera_Map_Back"), EKeys::Gamepad_FaceButton_Right},
+		{TEXT("VirtualCamera_Interaction_Back"), EKeys::Gamepad_FaceButton_Right},
+		{TEXT("VirtualCamera_Settings_Close"), EKeys::Gamepad_FaceButton_Right},
+		{TEXT("VirtualCamera_Focus_CloseMenu"), EKeys::Gamepad_FaceButton_Right},
+		{TEXT("VirtualCamera_MotionAdjustments_Close"), EKeys::Gamepad_Special_Right},
+		{TEXT("VirtualCamera_Settings_Close"), EKeys::Gamepad_Special_Right},
+		{TEXT("VirtualCamera_Focus_AdjustSelectorDown"), EKeys::Gamepad_DPad_Up},
+		{TEXT("VirtualCamera_Focus_AdjustSelectorUp"), EKeys::Gamepad_DPad_Down},
+		{TEXT("VirtualCamera_Focus_ButtonSelect"), EKeys::Gamepad_FaceButton_Bottom},
+		{TEXT("VirtualCamera_AnimationPreview_PlayPause"), EKeys::Gamepad_RightShoulder},
+		{TEXT("VirtualCamera_AnimationPreview_Close"), EKeys::Gamepad_FaceButton_Right},
+		{TEXT("VirtualCamera_AnimationPreview_BackStart"), EKeys::Gamepad_DPad_Left},
+		{TEXT("VirtualCamera_AnimationPreview_Close"), EKeys::Gamepad_Special_Right},
+		{TEXT("VirtualCamera_Reposition_ToggleTiltOffset"), EKeys::Gamepad_FaceButton_Left},
+		{TEXT("VirtualCamera_Reposition_Close"), EKeys::Gamepad_Special_Right},
+		{TEXT("VirtualCamera_Reposition_Close"), EKeys::Gamepad_FaceButton_Right},
+		{TEXT("VirtualCamera_MotionAdjustments_MoveAxisUp"), EKeys::Gamepad_DPad_Up},
+		{TEXT("VirtualCamera_MotionAdjustments_MoveAxisDown"), EKeys::Gamepad_DPad_Down},
+		{TEXT("VirtualCamera_MotionAdjustments_OpenNumpad"), EKeys::Gamepad_FaceButton_Bottom},
+		{TEXT("VirtualCamera_MotionAdjustments_SelectionSwitchLeft"), EKeys::Gamepad_LeftShoulder},
+		{TEXT("VirtualCamera_MotionAdjustments_SelectionSwitchRight"), EKeys::Gamepad_RightShoulder},
+		{TEXT("VirtualCamera_MotionAdjustments_LockAxis"), EKeys::Gamepad_FaceButton_Top},
+		{TEXT("VirtualCamera_MotionAdjustments_ResetDutch"), EKeys::Gamepad_FaceButton_Left},
+		{TEXT("VirtualCamera_MotionAdjustments_NumpadLeft"), EKeys::Gamepad_DPad_Left},
+		{TEXT("VirtualCamera_MotionAdjustments_NumpadRight"), EKeys::Gamepad_DPad_Right},
+		{TEXT("VirtualCamera_MotionAdjustments_Back"), EKeys::Gamepad_FaceButton_Right},
+
+	};
+	
+	AxisMappings = {
+		{TEXT("VirtualCamera_Home_MoveFocalLengthSelection"), EKeys::Gamepad_RightY,1.0f},
+		{TEXT("VirtualCamera_Home_MoveFStopSelection"), EKeys::Gamepad_LeftY,1.0f},
+		{TEXT("VirtualCamera_Playback_FastScroll"), EKeys::Gamepad_LeftY,1.0f},
+		{TEXT("VirtualCamera_Playback_FastScrub"), EKeys::Gamepad_RightX,1.0f},
+		{TEXT("VirtualCamera_Playback_BackwardFrame"), EKeys::Gamepad_LeftTriggerAxis,1.0f},
+		{TEXT("VirtualCamera_Playback_ForwardFrame"), EKeys::Gamepad_RightTriggerAxis,1.0f},
+		{TEXT("VirtualCamera_Map_Zoom"), EKeys::Gamepad_RightY,1.0f},
+		{TEXT("VirtualCamera_Map_XAxisPan"), EKeys::Gamepad_LeftX,1.0f},
+		{TEXT("VirtualCamera_Map_YAxisPan"), EKeys::Gamepad_LeftY,1.0f},
+		{TEXT("VirtualCamera_Interaction_Value_Change"), EKeys::Gamepad_LeftY,1.0f},
+		{TEXT("VirtualCamera_Interaction_Value_Change"), EKeys::Gamepad_RightY,1.0f},
+		{TEXT("VirtualCamera_Focus_ReticleMovementY"), EKeys::Gamepad_LeftY,1.0f},
+		{TEXT("VirtualCamera_Focus_ReticleMovementX"), EKeys::Gamepad_LeftX,1.0f},
+		{TEXT("VirtualCamera_Focus_AdjustFocalDistance"), EKeys::Gamepad_RightY,1.0f},
+		{TEXT("VirtualCamera_AnimationPreview_FastScrub"), EKeys::Gamepad_RightX,1.0f},
+		{TEXT("VirtualCamera_AnimationPreview_ForwardFrame"), EKeys::Gamepad_RightTriggerAxis,1.0f},
+		{TEXT("VirtualCamera_AnimationPreview_BackwardFrame"), EKeys::Gamepad_LeftTriggerAxis,1.0f},
+		{TEXT("VirtualCamera_MotionAdjustments_AdjustRadial"), EKeys::Gamepad_LeftY,1.0f},
+		{TEXT("VirtualCamera_MotionAdjustments_AdjustRadial"), EKeys::Gamepad_RightY,1.0f},
+		{TEXT("VirtualCamera_MotionAdjustments_AdjustRadial"), EKeys::Gamepad_LeftX,1.0f},
+		{TEXT("VirtualCamera_MotionAdjustments_AdjustRadial"), EKeys::Gamepad_RightX,1.0f},
+		{TEXT("VirtualCamera_Home_LeftMoveHorizontal"), EKeys::Gamepad_LeftX,1.0f},
+		{TEXT("VirtualCamera_Home_RightMoveHorizontal"), EKeys::Gamepad_RightX,1.0f},
+	};
+}
+
+void UVirtualCameraUserSettings::GetActionMappingsByName(const FName InActionName, TArray<FInputActionKeyMapping>& OutMappings) const
+{
+	if (InActionName.IsValid())
+	{
+		OutMappings = ActionMappings.FilterByPredicate([InActionName](const FInputActionKeyMapping& Mapping) 
+		{
+			return Mapping.ActionName == InActionName;
+		});
+
+	}
+}
+
+void UVirtualCameraUserSettings::GetAxisMappingsByName(const FName InAxisName, TArray<FInputAxisKeyMapping>& OutMappings) const
+{
+	if (InAxisName.IsValid())
+	{
+		OutMappings = AxisMappings.FilterByPredicate([InAxisName](const FInputAxisKeyMapping& Mapping)
+		{
+			return Mapping.AxisName == InAxisName;
+		});
+	}
+}
diff --git a/Engine/Plugins/Experimental/VirtualCamera/Source/VirtualCamera/Public/IVirtualCameraController.h b/Engine/Plugins/Experimental/VirtualCamera/Source/VirtualCamera/Public/IVirtualCameraController.h
new file mode 100644
--- /dev/null
+++ b/Engine/Plugins/Experimental/VirtualCamera/Source/VirtualCamera/Public/IVirtualCameraController.h
@@ -0,0 +1,158 @@
+// Copyright Epic Games, Inc. All Rights Reserved.
+
+#pragma once
+
+#include "CoreMinimal.h"
+#include "Delegates/DelegateCombinations.h"
+#if UE_ENABLE_INCLUDE_ORDER_DEPRECATED_IN_5_1
+#include "Engine/EngineTypes.h"
+#endif
+#include "Engine/HitResult.h"
+#include "LiveLinkRole.h"
+#include "Misc/CoreMiscDefines.h"
+#include "UObject/Interface.h"
+#include "UObject/ScriptInterface.h"
+
+#include "IVirtualCameraController.generated.h"
+
+class IBlendableInterface;
+class IVirtualCameraOptions;
+class IVirtualCameraPresetContainer;
+class UCineCameraComponent;
+class ULevelSequencePlaybackController;
+class USceneCaptureComponent2D;
+class UWorld;
+
+UENUM(BlueprintType)
+enum class EVirtualCameraFocusMethod : uint8
+{
+	/* Depth of Field disabled entirely */
+	None,
+
+	/* User controls focus distance directly */
+	Manual,
+
+	/* Focus distance is locked onto a specific point in relation to an actor */
+	Tracking,
+
+	/* Focus distance automatically changes to focus on actors in a specific screen location */
+	Auto,
+};
+
+USTRUCT(BlueprintType)
+struct FTrackingOffset
+{
+	GENERATED_BODY()
+
+	UPROPERTY(EditAnywhere, BlueprintReadWrite, Category = "Camera Settings")
+	FVector Translation;
+
+	UPROPERTY(EditAnywhere, BlueprintReadWrite, Category = "Camera Settings")
+	FRotator Rotation;
+
+	FTransform AsTransform() const
+	{
+		return FTransform(Rotation, Translation);
+	}
+
+	FTrackingOffset()
+		: Translation(EForceInit::ForceInitToZero), Rotation(EForceInit::ForceInitToZero)
+	{
+
+	};
+};
+
+USTRUCT(BlueprintType)
+struct FVirtualCameraTransform
+{
+	GENERATED_BODY()
+
+	UPROPERTY(EditAnywhere, BlueprintReadWrite, Category = "VirtualCamera")
+	FTransform Transform;
+};
+
+DECLARE_DYNAMIC_DELEGATE_RetVal_OneParam(FVirtualCameraTransform, FPreSetVirtualCameraTransform, FVirtualCameraTransform, CameraTransform);
+DECLARE_DYNAMIC_DELEGATE_OneParam(FOnActorClickedDelegate, FHitResult, HitResult);
+DECLARE_DYNAMIC_MULTICAST_DELEGATE_OneParam(FVirtualCameraTickDelegateGroup, float, DeltaTime);
+DECLARE_DYNAMIC_DELEGATE_OneParam(FVirtualCameraTickDelegate, float, DeltaTime);
+
+UINTERFACE(Blueprintable)
+class VIRTUALCAMERA_API UVirtualCameraController: public UInterface
+{
+	GENERATED_BODY()
+};
+
+class VIRTUALCAMERA_API IVirtualCameraController
+{
+	GENERATED_BODY()
+	
+public:
+
+	/** Returns the target camera that is used to create the streamed view. */
+	UFUNCTION(BlueprintNativeEvent, BlueprintCallable, Category = "VirtualCamera | Component")
+	UCineCameraComponent* GetStreamedCameraComponent() const;
+
+	UFUNCTION(BlueprintNativeEvent, BlueprintCallable, Category = "VirtualCamera | Component")
+	USceneCaptureComponent2D* GetSceneCaptureComponent() const;
+
+	/** Returns the VirtualCamera's Sequence Controller. */
+	UFUNCTION(BlueprintNativeEvent, BlueprintCallable, Category = "VirtualCamera | Sequencer")
+	ULevelSequencePlaybackController* GetSequenceController() const;
+
+	UFUNCTION(BlueprintNativeEvent, BlueprintCallable, Category = "VirtualCamera | Presets")
+	TScriptInterface<IVirtualCameraPresetContainer> GetPresetContainer();
+
+	UFUNCTION(BlueprintNativeEvent, BlueprintCallable, Category = "VirtualCamera | Options")
+	TScriptInterface<IVirtualCameraOptions> GetOptions();
+
+	UFUNCTION(BlueprintNativeEvent, BlueprintCallable, Category = "VirtualCamera | Movement")
+	FLiveLinkSubjectRepresentation GetLiveLinkRepresentation() const;
+
+	UFUNCTION(BlueprintNativeEvent, BlueprintCallable, Category = "VirtualCamera | Movement")
+	void SetLiveLinkRepresentation(const FLiveLinkSubjectRepresentation& InLiveLinkRepresenation);
+
+	virtual bool StartStreaming() PURE_VIRTUAL(IVirtualCameraController::StartStreaming, return false;);
+
+	virtual bool StopStreaming() PURE_VIRTUAL(IVirtualCameraController::StopStreaming, return false;);
+
+	virtual UWorld* GetControllerWorld() PURE_VIRTUAL(IVirtualCameraController::GetControllerWorld, return nullptr;);
+
+	UFUNCTION(BlueprintNativeEvent, BlueprintCallable, Category = "VirtualCamera | Streaming")
+	bool IsStreaming() const;
+
+	/** Check whether settings should save when stream is stopped. */
+	UFUNCTION(BlueprintNativeEvent, BlueprintCallable, Category = "VirtualCamera | Settings")
+	bool ShouldSaveSettingsOnStopStreaming() const;
+
+	/** Sets whether settings should be saved when stream is stopped. */
+	UFUNCTION(BlueprintNativeEvent, BlueprintCallable, Category = "VirtualCamera | Settings")
+	void SetSaveSettingsOnStopStreaming(bool bShouldSettingsSave);
+
+	UFUNCTION(BlueprintNativeEvent, BlueprintCallable, Category = "VirtualCamera | Movement")
+	FTransform GetRelativeTransform() const;
+
+	UFUNCTION(BlueprintNativeEvent, BlueprintCallable, Category = "VirtualCamera")
+	void AddBlendableToCamera(const TScriptInterface<IBlendableInterface>& InBlendableToAdd, float InWeight);
+
+	UFUNCTION(BlueprintNativeEvent, BlueprintCallable, Category = "VirtualCamera | Focus")
+	void SetTrackedActorForFocus(AActor* InActorToTrack, const FVector& TrackingPointOffset);
+
+	UFUNCTION(BlueprintNativeEvent, BlueprintCallable, Category = "VirtualCamera | Focus")
+	void SetFocusVisualization(bool bInShowFocusVisualization);
+
+	/** Delegate will be executed before transform is set onto VirtualCamera. */
+	UFUNCTION(BlueprintNativeEvent, BlueprintCallable, Category = "VirtualCamera | Movement")
+	void SetBeforeSetVirtualCameraTransformDelegate(const FPreSetVirtualCameraTransform& InDelegate);
+
+	/** Delegate will be executed when an actor in the scene was clicked/touched. */
+	UFUNCTION(BlueprintNativeEvent, BlueprintCallable, Category = "VirtualCamera | Focus")
+	void SetOnActorClickedDelegate(const FOnActorClickedDelegate& InDelegate);
+
+	/** Adds a delegate that will be executed every tick while streaming. */
+	UFUNCTION(BlueprintNativeEvent, BlueprintCallable, Category = "VirtualCamera")
+	void AddOnVirtualCameraUpdatedDelegate(const FVirtualCameraTickDelegate& InDelegate);
+
+	/** Remove delegate that is executed every tick while streaming. */
+	UFUNCTION(BlueprintNativeEvent, BlueprintCallable, Category = "VirtualCamera")
+	void RemoveOnVirtualCameraUpdatedDelegate(const FVirtualCameraTickDelegate& InDelegate);
+};
diff --git a/Engine/Plugins/Experimental/VirtualCamera/Source/VirtualCamera/Public/VirtualCameraActor.h b/Engine/Plugins/Experimental/VirtualCamera/Source/VirtualCamera/Public/VirtualCameraActor.h
new file mode 100644
--- /dev/null
+++ b/Engine/Plugins/Experimental/VirtualCamera/Source/VirtualCamera/Public/VirtualCameraActor.h
@@ -0,0 +1,203 @@
+// Copyright Epic Games, Inc. All Rights Reserved.
+
+#pragma once
+
+#include "CineCameraActor.h"
+#include "GameFramework/Actor.h"
+#include "IRemoteSessionRole.h"
+#include "IVirtualCameraController.h"
+#include "IVirtualCameraOptions.h"
+#include "IVirtualCameraPresetContainer.h"
+#include "LiveLinkRole.h"
+#include "Templates/UniquePtr.h"
+#include "UObject/NoExportTypes.h"
+#if WITH_EDITOR
+#include "UnrealEdMisc.h"
+#endif
+
+#include "VirtualCameraActor.generated.h"
+
+struct FVirtualCameraViewportSettings;
+class IRemoteSessionUnmanagedRole;
+class URemoteSessionMediaCapture;
+class URemoteSessionMediaOutput;
+class USceneCaptureComponent2D;
+class UUserWidget;
+class UVirtualCameraMovement;
+class UVPFullScreenUserWidget;
+class UWorld;
+
+#if WITH_EDITOR
+class UBlueprint;
+struct FAssetData;
+struct FCanDeleteAssetResult;
+#endif
+
+
+UCLASS(Abstract, BluePrintable, BlueprintType, Category="VirtualCamera", DisplayName="VirtualCameraActor")
+class VIRTUALCAMERA_API AVirtualCameraActor : public ACineCameraActor, public IVirtualCameraController, public IVirtualCameraPresetContainer, public IVirtualCameraOptions
+{
+	GENERATED_BODY()
+
+public:
+
+	AVirtualCameraActor(const FObjectInitializer& ObjectInitializer);
+	AVirtualCameraActor(FVTableHelper& Helper);
+	~AVirtualCameraActor();
+
+	UPROPERTY(VisibleAnywhere, BlueprintReadWrite, Category = "VirtualCamera | Component")
+	TObjectPtr<USceneCaptureComponent2D> SceneCaptureComponent;
+
+	UPROPERTY(EditAnywhere, BlueprintReadOnly, Category = "VirtualCamera | Movement")
+	FLiveLinkSubjectRepresentation LiveLinkSubject;
+
+	UPROPERTY(VisibleAnywhere, BlueprintReadWrite, Category = "VirtualCamera | Movement")
+	TObjectPtr<UVirtualCameraMovement> MovementComponent;
+
+	UPROPERTY(Transient, EditDefaultsOnly, Category = "VirtualCamera | MediaOutput")
+	TObjectPtr<URemoteSessionMediaOutput> MediaOutput;
+
+	UPROPERTY(EditAnywhere, BlueprintReadWrite, Category = "VirtualCamera | UMG")
+	TSubclassOf<UUserWidget> CameraUMGClass;
+
+	UPROPERTY(EditAnywhere, BlueprintReadWrite, Category = "VirtualCamera | Streaming")
+	FVector2D TargetDeviceResolution;
+
+	UPROPERTY(EditAnywhere, BlueprintReadWrite, Category = "VirtualCamera | Streaming")
+	int32 RemoteSessionPort;
+
+protected:
+
+	UPROPERTY(EditDefaultsOnly, Category = "VirtualCamera | UMG")
+	TObjectPtr<UVPFullScreenUserWidget> CameraScreenWidget;
+
+	UPROPERTY(Transient, EditDefaultsOnly, Category = "VirtualCamera | MediaOutput")
+	TObjectPtr<URemoteSessionMediaCapture> MediaCapture;
+
+	UPROPERTY(Transient)
+	TObjectPtr<UWorld> ActorWorld;
+
+	UPROPERTY(Transient)
+	TObjectPtr<AActor> PreviousViewTarget;
+
+	/** Should focus plane be shown on all touch focus events */
+	UPROPERTY(BlueprintReadOnly, Category = "VirtualCamera | Focus")
+	bool bAllowFocusVisualization;
+
+	/**
+	 * Delegate that will is triggered before transform is set onto Actor.
+	 * @param FVirtualCameraTransform Transform data that is passed to delegate.
+	 * @return FVirtualCameraTransform Manipulated transform that will be set onto Actor.
+	 */
+	UPROPERTY(EditAnywhere, Category = "VirtualCamera")
+	FPreSetVirtualCameraTransform OnPreSetVirtualCameraTransform;
+
+	/**
+	 * Delegate that will be triggered when an actor has been clicked/touched.
+	 * @note Delegate will run on Touch/Mouse-Down
+	 * @param AActor* Pointer to the actor that was clicked.
+	 */
+	UPROPERTY(EditAnywhere, Category = "VirtualCamera | Focus")
+	FOnActorClickedDelegate OnActorClickedDelegate;
+
+	/**
+	 * This delegate is triggered at the end of a tick in editor/pie/game.
+	 * @note The Actor is only ticked while it is being streamed.
+	 * @param float Delta Time in seconds.
+	 */
+	UPROPERTY(EditAnywhere, Category = "VirtualCamera")
+	FVirtualCameraTickDelegateGroup OnVirtualCameraUpdatedDelegates;
+
+	/** The next preset number */
+	static int32 PresetIndex;
+
+	/* Stores the list of settings presets, and saved presets */
+	UPROPERTY(EditAnywhere, Category = "VirtualCamera | Settings")
+	TMap<FString, FVirtualCameraSettingsPreset> SettingsPresets;
+
+	/** The desired unit in which to display focus distance */
+	EUnit DesiredDistanceUnits;
+
+	/** Whether to save all settings when streaming is stopped */
+	bool bSaveSettingsOnStopStreaming;
+
+public:
+
+	//~ Begin AActor interface
+	virtual void Destroyed() override;
+#if WITH_EDITOR
+	virtual bool ShouldTickIfViewportsOnly() const override;
+#endif
+	virtual void Tick(float DeltaSeconds) override;
+
+protected:
+
+	virtual void BeginPlay() override;
+	virtual void EndPlay(const EEndPlayReason::Type EndPlayReason) override;
+	//~ End AActor interface
+
+	//~ Begin IVirtualCameraController Interface
+public:
+	virtual bool StartStreaming() override;
+	virtual bool StopStreaming() override;
+	virtual UWorld* GetControllerWorld() override;
+protected:
+	virtual UCineCameraComponent* GetStreamedCameraComponent_Implementation() const override;
+	virtual USceneCaptureComponent2D* GetSceneCaptureComponent_Implementation() const override;
+	virtual ULevelSequencePlaybackController* GetSequenceController_Implementation() const override;
+	virtual TScriptInterface<IVirtualCameraPresetContainer> GetPresetContainer_Implementation() override;
+	virtual TScriptInterface<IVirtualCameraOptions> GetOptions_Implementation() override;
+	virtual FLiveLinkSubjectRepresentation GetLiveLinkRepresentation_Implementation() const override;
+	virtual void SetLiveLinkRepresentation_Implementation(const FLiveLinkSubjectRepresentation& InLiveLinkRepresentation) override;
+	virtual bool IsStreaming_Implementation() const override;
+	virtual void SetSaveSettingsOnStopStreaming_Implementation(bool bShouldSettingsSave) override;
+	virtual FTransform GetRelativeTransform_Implementation() const override;
+	virtual void AddBlendableToCamera_Implementation(const TScriptInterface<IBlendableInterface>& InBlendableToAdd, float InWeight) override;
+	virtual void SetTrackedActorForFocus_Implementation(AActor* InActorToTrack, const FVector& InTrackingPointOffset) override;
+	virtual void SetFocusVisualization_Implementation(bool bInShowFocusVisualization) override;
+	virtual bool ShouldSaveSettingsOnStopStreaming_Implementation() const override;
+	virtual void SetBeforeSetVirtualCameraTransformDelegate_Implementation(const FPreSetVirtualCameraTransform& InDelegate) override;
+	virtual void SetOnActorClickedDelegate_Implementation(const FOnActorClickedDelegate& InDelegate) override;
+	virtual void AddOnVirtualCameraUpdatedDelegate_Implementation(const FVirtualCameraTickDelegate& InDelegate) override;
+	virtual void RemoveOnVirtualCameraUpdatedDelegate_Implementation(const FVirtualCameraTickDelegate& InDelegate) override;
+	//~ End  IVirtualCameraController Interface
+
+	//~ Begin IVirtualCameraPresetContainer Interface
+protected:
+	virtual FString SavePreset_Implementation(const bool bSaveCameraSettings, const bool bSaveStabilization, const bool bSaveAxisLocking, const bool bSaveMotionScale) override;
+	virtual bool LoadPreset_Implementation(const FString& PresetName) override;
+	virtual int32 DeletePreset_Implementation(const FString& PresetName) override;
+	virtual TMap<FString, FVirtualCameraSettingsPreset> GetSettingsPresets_Implementation() override;
+	//~ End IVirtualCameraPresetContainer Interface
+
+	//~ Begin IVirtualCameraOptions Interface
+protected:
+	virtual void SetDesiredDistanceUnits_Implementation(const EUnit DesiredUnits) override;
+	virtual EUnit GetDesiredDistanceUnits_Implementation() override;
+	virtual bool IsFocusVisualizationAllowed_Implementation() override;
+	//~ End IVirtualCameraOptions Interface
+
+private:
+	void OnTouchEventOutsideUMG(const FVector2D& InViewportPosition);
+
+	/** Stores the current camera settings to a save game for later use. */
+	void SaveSettings();
+	/** Restores settings from save game. */
+	void LoadSettings();
+
+#if WITH_EDITOR
+	void OnMapChanged(UWorld* World, EMapChangeType ChangeType);
+	void OnBlueprintPreCompile(UBlueprint* Blueprint);
+	void OnPrepareToCleanseEditorObject(UObject* Object);
+	void OnAssetRemoved(const FAssetData& AssetData);
+	void OnAssetsCanDelete(const TArray<UObject*>& InAssetsToDelete, FCanDeleteAssetResult& CanDeleteResult);
+#endif
+
+private:
+
+	UPROPERTY(BlueprintGetter = GetCineCameraComponent, Category = "VirtualCamera | Component")
+	TObjectPtr<UCineCameraComponent> StreamedCamera;
+	bool bIsStreaming;
+	TUniquePtr<FVirtualCameraViewportSettings> ViewportSettingsBackup;
+	FHitResult LastViewportTouchResult;
+};
diff --git a/Engine/Plugins/Experimental/VirtualCamera/Source/VirtualCamera/Public/VirtualCameraCineCameraComponent.h b/Engine/Plugins/Experimental/VirtualCamera/Source/VirtualCamera/Public/VirtualCameraCineCameraComponent.h
new file mode 100644
--- /dev/null
+++ b/Engine/Plugins/Experimental/VirtualCamera/Source/VirtualCamera/Public/VirtualCameraCineCameraComponent.h
@@ -0,0 +1,234 @@
+// Copyright Epic Games, Inc. All Rights Reserved.
+
+#pragma once
+
+#include "CineCameraComponent.h"
+#include "IVirtualCameraController.h"
+
+#include "VirtualCameraCineCameraComponent.generated.h"
+
+UCLASS(Blueprintable, BlueprintType)
+class VIRTUALCAMERA_API UVirtualCameraCineCameraComponent : public UCineCameraComponent
+{
+	GENERATED_UCLASS_BODY()
+
+public:
+	/** List of preset aperture options, aperture will always be one of these values */
+	UPROPERTY(EditAnywhere, BlueprintReadWrite, Category = "Current Camera Settings")
+	TArray<float> ApertureOptions;
+
+	/** List of preset focal length options, focal length will be one of these values, unless manually zooming */
+	UPROPERTY(EditAnywhere, BlueprintReadWrite, Category = "Current Camera Settings")
+	TArray<float> FocalLengthOptions;
+
+	/** List of preset matte options to chose from, UI options will only pull from this, unless a filmback option with a custom matte is selected */
+	UPROPERTY(EditAnywhere, BlueprintReadWrite, Category = "Current Camera Settings")
+	TArray<float> MatteOptions;
+
+	/** List of preset filmback options, filmback will always be one of these values */
+	UPROPERTY(EditAnywhere, Category = "Current Camera Settings")
+	TMap<FString, FCameraFilmbackSettings> FilmbackOptions;
+
+	/** The desired filmback settings to be shown in the viewport within Virtual Camera UI window */
+	UPROPERTY(EditAnywhere, BlueprintReadOnly, Category = "Current Camera Settings")
+	FCameraFilmbackSettings DesiredFilmbackSettings;
+
+	/** The filmback settings to be used for additional letterboxing if desired */
+	UPROPERTY(EditAnywhere, BlueprintReadOnly, Category = "Current Camera Settings")
+	float MatteAspectRatio;
+
+	/** The opacity of the matte in the camera view */
+	UPROPERTY(BlueprintReadWrite, Category = "Current Camera Settings", meta = (ClampMin = "0.0", ClampMax = "1.0"))
+	float MatteOpacity;
+
+	/** The X and Y ratios of Desired View Size to Actual View Size (calculated as Desired/ Actual) */
+	UPROPERTY(BlueprintReadWrite, Category = "Current Camera Settings")
+	FVector2D ViewSizeRatio;
+
+	/** Tracks if autofocus is enabled */
+	bool bAutoFocusEnabled;
+
+	/** Tracks whether or not the camera's view needs to be restored */
+	bool bCameraViewResetNeeded;
+
+	virtual void BeginPlay() override;
+
+	/**
+	 * Get the current focal length value on the camera.
+	 * @return the current focal length of the camera being viewed
+	 */
+	float GetCurrentFocalLength() const { return CurrentFocalLength; }
+
+	/**
+	 * Sets the current focal length of the cinematic camera to a given value.
+	 * @param NewFocalLength - The new target focal length for the camera
+	 */
+	void SetCurrentFocalLength(const float NewFocalLength) { CurrentFocalLength = NewFocalLength; }
+
+	/**
+	 * Adjust the focal length of the camera to an adjacent preset value.
+	 * @param bShiftUp - If true, the focal length will increase; if false, it will decrease
+	 * @return the new focal length of the camera in mm
+	 */
+	float ChangeFocalLengthPreset(const bool bShiftUp);
+
+	/**
+	 * Get the current aperture value on the camera.
+	 * @return the current aperture of the camera component
+	 */
+	float GetCurrentAperture() const { return CurrentAperture; }
+
+	/**
+	 * Set the current aperture value on the camera.
+	 * @param NewAperture - The new aperture value
+	 */
+	void SetCurrentAperture(const float NewAperture) { CurrentAperture = NewAperture; }
+
+	/**
+	 * Adjusts the aperture of the camera to an adjacent preset value.
+	 * @param bShiftUp - If true, the aperture will increase; if false, it will decrease
+	 * @return the new aperture of the camera
+	 */
+	float ChangeAperturePreset(const bool bShiftUp);
+
+	/**
+	 * Get the name of the current preset filmback option on the camera.
+	 * @return the name of the current preset filmback option
+	 */
+	FString GetCurrentFilmbackName() const { return CurrentFilmbackOptionsKey; };
+
+	/**
+	 * Get the currently used focus method.
+	 * @return the Current focus method
+	 */
+	EVirtualCameraFocusMethod GetCurrentFocusMethod() const { return CurrentFocusMethod; };
+
+	/**
+	 * Stores the names of all available filmback presets into an array.
+	 * @param OutFilmbackPresets - Upon return, will contain all available filmback presets.
+	 * @return true if operation was successful
+	 */
+	bool GetFilmbackPresetOptions(TArray<FString>& OutFilmbackPresetsArray) const;
+
+	/**
+	 * Set the filmback settings to a new filmback preset.
+	 * @param NewFilmbackPreset - The name of the desired preset to use
+	 * @return true if operation was successful; false if NewFilmbackPreset is not a valid option
+	 */
+	bool SetFilmbackPresetOption(const FString& NewFilmbackPreset);
+
+	/**
+	 * Returns the values of all matte options.
+	 * @param &OutMatteValues - Upon return, array will contain all matte values.
+	 */
+	void GetMatteValues(TArray<float>& OutMatteValues) const ;
+
+	/**
+	 * Returns the current matte aspect ratio.
+	 * @return the current matte setting
+	 */
+	float GetMatteAspectRatio() { return MatteAspectRatio; }
+
+	/**
+	 * Set the matte aspect ratio to a new value.
+	 * @param NewMatteAspectRatio - The desired matte aspect to use
+	 * @return true if operation was successful; false if NewMattePreset is not a valid input
+	 */
+	bool SetMatteAspectRatio(const float NewMatteAspectRatio);
+
+	/**
+	 * Sets the camera focus method.
+	 * @param NewFocusMethod - The new focus method to be used by the camera
+	 */
+	void SetFocusMethod(const EVirtualCameraFocusMethod NewFocusMethod);
+
+	/**
+	 * Sets the focus change rate.
+	 * @param NewFocusChangeSmoothness - The new change rate to use for smooth camera focus shifting (Clamped between 0 and 1)
+	 */
+	void SetFocusChangeSmoothness(const float NewFocusChangeSmoothness);
+
+	/**
+	 * Gets the current focus distance of the camera.
+	 * @return the focal distance of the camera
+	 */
+	float GetCurrentFocusDistance() { return CurrentFocusDistance; }
+
+	/**
+	 * Sets the current focus distance of the camera to a new value.
+	 * @param NewFocusDistance - The new focus distance for the camera
+	 */
+	void SetCurrentFocusDistance(const float NewFocusDistance) { FocusSettings.ManualFocusDistance = NewFocusDistance; }
+
+	/**
+	 * Sets the current focus distance to a new value based on current focus settings.
+	 * @param NewFocusDistance - The desired focus distance
+	 */
+	void SetFocusDistance(const float NewFocusDistance);
+
+	/**
+	 * Sets tracked actor settings on the camera.
+	 * @param ActorToTrack - The actor that the focus should follow
+	 * @param TrackingPointOffset - The offset from the root of the actor to the point to focus on
+	 */
+	void SetTrackedActorForFocus(AActor* ActorToTrack, const FVector TrackingPointOffset);
+
+	/**
+	 * Adds a blendable object to the camera's post process settings.
+	 * @param BlendableToAdd - The blendable that will be added to the post process settings
+	 * @param Weight - The weighting of the blendable's alpha
+	 */
+	void AddBlendableToCamera(const TScriptInterface<IBlendableInterface> BlendableToAdd, const float Weight);
+
+	/**
+	 * Toggles focus visualization tools on camera.
+	 * @param bShowFocusVisualization - The desired state of the visualization tools
+	 */
+	void SetFocusVisualization(bool bShowFocusVisualization);
+
+	/**
+	 * Is focus visualization tools activated on camera.
+	 */
+	bool IsFocusVisualizationActivated() const;
+
+	/**
+	 * Updates the camera view to have the desired film format view within the Virtual Camera UI area.
+	 */
+	void UpdateCameraView();
+
+	/**
+	 * Allow camera view updates.
+	 */
+	void AllowCameraViewUpdates() { bAllowCameraViewUpdates = true; }
+
+	/**
+	 * Disable camera view updates.
+	 */
+	void StopCameraViewUpdates() { bAllowCameraViewUpdates = false; }
+
+protected:
+	/** The current filmback option preset being used */
+	FString CurrentFilmbackOptionsKey;
+
+	/** Mesh used for the focus plane */
+	UStaticMesh* FocusPlaneMesh;
+
+	/** How smooth focus changes are, set by user through UI */
+	float FocusChangeSmoothness;
+
+	/** Saves the current focus method for reset if menus is exited without applying settings */
+	EVirtualCameraFocusMethod CachedFocusMethod;
+
+	/** The current focus method */
+	EVirtualCameraFocusMethod CurrentFocusMethod;
+
+	/** Whether or not camera view updates should occur */
+	bool bAllowCameraViewUpdates;
+
+	/**
+	 * Searches a preset array for the closest value to it.
+	 * @param SearchValue - The value that will be used to find the closest preset
+	 * @return the index of the closest value to the search value in the preset array
+	 */
+	int32 FindClosestPresetIndex(const TArray<float>& ArrayToSearch, const float SearchValue) const;
+};
diff --git a/Engine/Plugins/Experimental/VirtualCamera/Source/VirtualCamera/Public/VirtualCameraConcertCameraComponent.h b/Engine/Plugins/Experimental/VirtualCamera/Source/VirtualCamera/Public/VirtualCameraConcertCameraComponent.h
new file mode 100644
--- /dev/null
+++ b/Engine/Plugins/Experimental/VirtualCamera/Source/VirtualCamera/Public/VirtualCameraConcertCameraComponent.h
@@ -0,0 +1,51 @@
+// Copyright Epic Games, Inc. All Rights Reserved.
+
+#pragma once
+
+#include "Components/ActorComponent.h"
+#include "VirtualCameraConcertCameraComponent.generated.h"
+
+class UCineCameraComponent;
+
+/**
+ * A class to transfer Camera data in MU session
+ */
+UCLASS(deprecated, Blueprintable, ClassGroup = Camera, meta = (BlueprintSpawnableComponent))
+class VIRTUALCAMERA_API UDEPRECATED_VirtualCameraConcertCameraComponent : public UActorComponent
+{
+	GENERATED_BODY()
+
+public:
+	UDEPRECATED_VirtualCameraConcertCameraComponent(const FObjectInitializer& ObjectInitializer);
+
+	/**
+	 * The tracked name used by the multi user system to send the data.
+	 * When not set the component name will be used. That may conflict with another instance that has the same name but within different actors.
+	 */
+	UPROPERTY(BlueprintReadWrite, EditInstanceOnly, Category = "VirtualCamera")
+	FName TrackingName;
+
+	/** Is that component sending the camera information or receiving it. */
+	UPROPERTY(BlueprintReadWrite, VisibleAnywhere, Category = "VirtualCamera")
+	bool bHasAuthority;
+
+	/**
+	 * When the component has the authority and we are in the editor, we should broadcast the update.
+	 * @note The component needs to be owned by an actor that is ticked in the editor.
+	 * @see ShouldTickIfViewportsOnly
+	 */
+	UPROPERTY(BlueprintReadWrite, EditAnywhere, Category = "VirtualCamera", meta = (EditCondition = "bHasAuthority"))
+	bool bSendUpdateInEditor;
+
+	/** Should update the camera component relative location & rotation. */
+	UPROPERTY(BlueprintReadWrite, EditAnywhere, Category = "VirtualCamera", meta = (EditCondition = "!bHasAuthority"))
+	bool bUpdateCameraComponentTransform;
+
+	/** The camera component that will be tracked. */
+	UPROPERTY(BlueprintReadWrite, EditAnywhere, Category = "VirtualCamera")
+	TObjectPtr<UCineCameraComponent> ComponentToTrack;
+
+public:
+	virtual void BeginPlay() override;
+	virtual void TickComponent(float DeltaTime, enum ELevelTick TickType, FActorComponentTickFunction *ThisTickFunction) override;
+};
diff --git a/Engine/Plugins/Experimental/VirtualCamera/Source/VirtualCamera/Public/VirtualCameraMovementComponent.h b/Engine/Plugins/Experimental/VirtualCamera/Source/VirtualCamera/Public/VirtualCameraMovementComponent.h
new file mode 100644
--- /dev/null
+++ b/Engine/Plugins/Experimental/VirtualCamera/Source/VirtualCamera/Public/VirtualCameraMovementComponent.h
@@ -0,0 +1,265 @@
+// Copyright Epic Games, Inc. All Rights Reserved.
+
+#pragma once
+
+#include "GameFramework/PawnMovementComponent.h"
+#include "VirtualCameraSaveGame.h"
+#include "VirtualCameraMovementComponent.generated.h"
+
+DECLARE_DYNAMIC_MULTICAST_DELEGATE(FVirtualCameraResetOffsetsDelegate);
+
+UCLASS()
+class VIRTUALCAMERA_API UVirtualCameraMovementComponent : public UPawnMovementComponent
+{
+	GENERATED_UCLASS_BODY()
+
+public:
+	/** Stores all settings for each axis */
+	TMap<EVirtualCameraAxis, FVirtualCameraAxisSettings> AxisSettings;
+
+	/** Tracks whether user wants to cache a locking axis set or always use local */
+	bool bLockRelativeToFirstLockAxis;
+
+	/** Tracks if boom should be global or relative to the camera */
+	bool bUseGlobalBoom;
+
+	/** Tracks if dutch should be reset to zero when freezing view */
+	bool bZeroDutchOnLock;
+
+	/** Delegate to broadcast when tracker offsets are reset */
+	FVirtualCameraResetOffsetsDelegate OnOffsetReset;
+
+	/**
+	 * Adds the given vector to the accumulated input in world space. Input vectors are usually between 0 and 1 in magnitude.
+	 * They are accumulated during a frame then applied as acceleration during the movement update.
+	 * @param WorldDirection - Direction in world space to apply input
+	 * @param bForce - If true always add the input, ignoring the result of IsMoveInputIgnored().
+	 * @see APawn::AddMovementInput()
+	 */
+	virtual void AddInputVector(FVector WorldVector, bool bForce = false) override;
+
+	/**
+	 * Adds the given vector to the accumulated input in world space. Similar to AddInputVector().
+	 * This calls input scaling specific to how controller (joystick) input should be handled.
+	 * @param WorldVector - Direction in world space to apply input.
+	 * @param MovementScaleAxis - Which axis movement scale to use.
+	 */
+	void AddInputVectorFromController(FVector WorldVector, EVirtualCameraAxis MovementScaleAxis);
+
+	/**
+	 * Update the position and rotation of the camera.
+	 * @param Location - The current location of the tracker being used for input
+	 * @param Rotation - The current rotation of the tracker being used for input
+	 */
+	void ProcessMovementInput(const FVector& Location, const FRotator& Rotation);
+
+	/**
+	 * Toggles the lock on a given axis; returns new locked state.
+	 * @param AxisToToggle - The axis whose lock should be turned on/off
+	 * @return the new locked state of AxisToToggle (true = locked)
+	 */
+	bool ToggleAxisLock(const EVirtualCameraAxis AxisToToggle);
+
+	/**
+	 * Toggles the freeze on a given axis; returns new frozen state.
+	 * @param AxisToToggle - The axis whose lock should be toggled
+	 * @return the new frozen state of AxisToToggle (true = frozen)
+	 */
+	bool ToggleAxisFreeze(const EVirtualCameraAxis AxisToToggle);
+
+	/**
+	 * Sets the stabilization rate for a given lock.
+	 * @param AxisToAdjust - The axis whose stabilization rate should be changed
+	 * @param NewStabilizationAmount - The stabilization amount we should attempt to set the value
+	 * @return the actual value the stabilization amount was set to after clamping
+	 */
+	float SetAxisStabilizationScale(const EVirtualCameraAxis AxisToAdjust, float NewStabilizationAmount);
+
+	/**
+	 * Gets the stabilization rate for a given axis. 
+	 * This method returns the stabilization value after undoing the curve applied in SetAxisStabiliaztionScale().
+	 * The return value will be equivialent the the NewStabiliaztionAmount parameter when setting the value, and not the returned actual value.
+	 * @param AxisToRetrieve - The axis whose stabilization rate is being retrieved
+	 * @return the value of the stabilization amount for the current axis
+	 */
+	float GetAxisStabilizationScale(const EVirtualCameraAxis AxisToRetrieve) const;
+
+	/**
+	 * Gets the movement scale for a given axis.
+	 * @param AxisToRetrieve - The axis whose movement scale is being retrieved
+	 * @return the value of the movement scale for the current axis
+	 */
+	float GetAxisMovementScale(const EVirtualCameraAxis AxisToRetrieve) const { return AxisSettings[AxisToRetrieve].MovementScale; };
+
+	/**
+	 * Moves the camera back to actor root and aligns rotation with the input tracker.
+	 */
+	void ResetCameraOffsetsToTracker();
+
+	/**
+	 * Sets the movement scale of the camera actor
+	 * @param AxisToAdjust - The axis to set movement scale on (setting for rotation axes is allowed but has no effect)
+	 * @param NewMovementScale - The desired scaling factor for calculating movement along this axis
+	 */
+	void SetMovementScale(const EVirtualCameraAxis AxisToAdjust, const float NewMovementScale) { AxisSettings[AxisToAdjust].MovementScale = NewMovementScale; }
+
+	/**
+	 * Gets the movement scale of the camera actor
+	 * @param AxisToRetrieve - The axis to get movement scale from (setting for rotation axes is allowed but has no effect)
+	 * @return the scaling factor for calculating movement along this axis
+	 */
+	float GetMovementScale(const EVirtualCameraAxis AxisToRetrieve) const { return AxisSettings[AxisToRetrieve].MovementScale; }
+
+	/**
+	 * Check to see if any location locks are active
+	 * @return true if any of the three location axes are locked
+	 */
+	bool IsLocationLockingActive() const;
+
+	/**
+	 * Checks if an axis is locked.
+	 * @param AxisToCheck - The axis being checked.
+	 * @return true if locked false otherwise
+	 */
+	bool IsAxisLocked(const EVirtualCameraAxis AxisToCheck) const;
+
+	/**
+	 * Set the value for the option to zero out dutch when locking that axis.
+	 * @param bInValue - The new setting to use for zeroing dutch on lock
+	 */
+	void SetZeroDutchOnLock(const bool bInValue) { bZeroDutchOnLock = bInValue; }
+
+	/**
+	 * Get the value for the option to zero out dutch when locking that axis.
+	 * @return true if dutch will be zeroed when that axis is locked
+	 */
+	bool IsDutchZeroingOnLockActive() const { return bZeroDutchOnLock; }
+
+	/** 
+	 * Handle a move forward command from the owner.
+	 * @param InValue - The driving axis value for movement
+	 */
+	void OnMoveForward(const float InValue);
+
+	/**
+	 * Handle a move right command from the owner.
+	 * @param InValue - The driving axis value for movement
+	 */
+	void OnMoveRight(const float InValue);
+
+	/**
+	 * Handle a move up command from the owner.
+	 * @param InValue - The driving axis value for movement
+	 */
+	void OnMoveUp(const float InValue);
+
+	/**
+	 * Teleport to a provided location and rotation. 
+	 * @param TargetTransform - The information for the target to teleport to
+	 */
+	void Teleport(const FTransform& TargetTransform);
+
+	/**
+	 * Get whether or not global boom is being used when navigating with the joysticks
+	 * @param bShouldUseGlobalBoom - True if using global Z axis, false if using local Z axis
+	 */
+	bool IsUsingGlobalBoom() const { return bUseGlobalBoom; }
+
+	/**
+	 * Sets whether or not global boom should be used when navigating with the joysticks
+	 * @param bShouldUseGlobalBoom - True to use global Z axis, false to use local Z axis
+	 */
+	void SetUseGlobalBoom(bool bShouldUseGlobalBoom) { bUseGlobalBoom = bShouldUseGlobalBoom; }
+
+
+	/**
+	 * When the input come from a physical controller, move this component instead.
+	 * This allow a proper replication when we are in a multi user setting.
+	 * This component doesn't respect the freeze and lock axis.
+	 */
+	void SetRootComponent(USceneComponent* FromController);
+
+private:
+	/** The cached axis for locking, only used if bLockRelativeToFirstLockAxis is true */
+	FQuat CachedLockingAxis;
+
+	/** Tracks the previous tracker location to determine which direction we should move */
+	FVector PreviousTrackerLocation;
+
+	/** Tracks the current target location of the camera for stabilization */
+	FVector TargetLocation;
+
+	/** Tracks the current target location of the camera that came from a controller */
+	FVector FromControllerTargetLocation;
+
+	/** Tracks the previous tracker rotation to determine how much we should rotate */
+	FRotator PreviousTrackerRotation;
+
+	/** Tracks the current target rotation of the camera for stabilization */
+	FRotator TargetRotation;
+
+	/**
+	 * The component we move and update when the input is coming from the controller.
+	 * If this is null UpdatedComponent will be used instead.
+	 * @see UpdatedComponent
+	 */
+	UPROPERTY(Transient, DuplicateTransient)
+	TObjectPtr<USceneComponent> RootUpdatedComponent;
+
+	/**
+	 * RootUpdatedComponent, cast as a UPrimitiveComponent. May be invalid if RootUpdatedComponent was null or not a UPrimitiveComponent.
+	 */
+	UPROPERTY(Transient, DuplicateTransient)
+	TObjectPtr<UPrimitiveComponent> RootUpdatedPrimitive;
+
+	/**
+	 * Applies relative axis filtering based on locks for location.
+	 * @param InVector - The input vector that will be used to apply locking offsets
+	 */
+	void ApplyLocationLocks(const FVector& InVector);
+
+	/**
+	 * Get the position filtering vector, which is used to filter each axis independently.
+	 * @return filtered location based on stabilization values
+	 */
+	FVector GetStabilizedDeltaLocation() const;
+
+	/**
+	 * Applies relative axis scaling for location.
+	 * @param VectorToAdjust - Upon return, has input values with local space movement scaled by movement scale
+	 */
+	void ApplyLocationScaling(FVector& VectorToAdjust);
+
+	/**
+	 * Applies global axis filtering based on locks for rotation.
+	 * @param InRotation - The input rotation that will be used to apply rotation offsets
+	 */
+	void ApplyRotationLocks(const FRotator& InRotation);
+
+	/**
+	 * Get the position filtering vector, which is used to filter each axis independently.
+	 * @return filtered rotation
+	 */
+	FRotator GetStabilizedRotation() const;
+
+	/**
+	 * Get the directional vectors for movement controls, taking into account any cached axes as needed.
+	 * @param OutForwardVector - Upon return, will hold the forward vector
+	 * @param OutRightVector - Upon return, will hold the right vector
+	 * @param OutUpVector - Upon return, will hold the up vector
+	 * @param bAllowLock - Use the locking axis (if valid) or the camera axis
+	 */
+	void GetDirectionVectorsForCamera(FVector& OutForwardVector, FVector& OutRightVector, FVector& OutUpVector, bool bTryLock) const;
+
+	/**
+	 * Gets the current location offset
+	 * @return - The current location offset
+	 */
+	FVector GetLocationOffset() const;
+
+	/**
+	 * Gets the current rotation offset
+	 * @return - The current rotation offset
+	 */
+	FRotator GetRotationOffset() const;
+};
