Move VirtualCamera and VirtualCameraEditor to VirtualCamera plugin

diff --git a/Engine/Plugins/Experimental/VirtualCamera/Source/VirtualCamera/Public/VirtualCameraPawnBase.h b/Engine/Plugins/Experimental/VirtualCamera/Source/VirtualCamera/Public/VirtualCameraPawnBase.h
new file mode 100644
--- /dev/null
+++ b/Engine/Plugins/Experimental/VirtualCamera/Source/VirtualCamera/Public/VirtualCameraPawnBase.h
@@ -0,0 +1,341 @@
+// Copyright Epic Games, Inc. All Rights Reserved.
+
+#pragma once
+
+#include "GameFramework/Pawn.h"
+#include "IVirtualCameraOptions.h"
+#include "IVirtualCameraPresetContainer.h"
+#include "VirtualCameraSaveGame.h"
+#include "VirtualCameraCineCameraComponent.h"
+#include "VirtualCameraMovementComponent.h"
+#include "VirtualCameraPawnBase.generated.h"
+
+// Forward Declarations
+class USceneComponent;
+
+/**
+ * A class to handle aspects of virtual Camera related to general settings, and communicating with components.
+ */
+UCLASS(Abstract, Blueprintable, BlueprintType, HideCategories = ("Pawn", "Camera", "Rendering", "Replication", "Input", "Actor", "HLOD"))
+class VIRTUALCAMERA_API AVirtualCameraPawnBase : public APawn, public IVirtualCameraPresetContainer, public IVirtualCameraOptions
+{
+	GENERATED_UCLASS_BODY()
+
+public:
+	virtual void BeginPlay() override;
+	virtual void EndPlay(const EEndPlayReason::Type EndPlayReason) override;
+
+	/** Root component */
+	USceneComponent* DefaultSceneRoot;
+
+	/** Cinematic camera used for view */
+	UPROPERTY(VisibleAnywhere, BlueprintReadWrite, Category = "Cinematic Camera")
+	TObjectPtr<UVirtualCameraCineCameraComponent> CineCamera;
+
+	/** Movement Component to handle the motion input for the camera */
+	UVirtualCameraMovementComponent* MovementComponent;
+
+	/** Determines if values should be saved between sessions */
+	UPROPERTY(EditAnywhere, Category = "Default")
+	bool bSaveSettingsWhenClosing;
+
+	/** Stores the name of the save slot being used currently */
+	UPROPERTY(EditAnywhere, Category = "Default")
+	FString SavedSettingsSlotName;
+
+	/* Stores the list of settings presets, and saved presets */
+	UPROPERTY(EditAnywhere, Category = "Default")
+	TMap<FString, FVirtualCameraSettingsPreset> SettingsPresets;
+
+	/**
+	 * Returns the information associated with a waypoint.
+	 * @param WaypointName - The name of the waypoint to retrieve
+	 * @param OutWaypointInfo - Upon return, will hold the info for that waypoint.
+	 */
+	void GetWaypointInfo(const FString WaypointName, FVirtualCameraWaypoint &OutWaypointInfo) const;
+
+	/**
+	 * Collects a list of existing waypoint names into an out array.
+	 * @param OutArray - Upon return, will store the waypoint names.
+	 */
+	void GetWaypointNames(TArray<FString>& OutArray) const;
+
+	/**
+	 * Returns the information associated with a Screenshot.
+	 * @param ScreenshotName - The name of the screenshot to retrieve
+	 * @param OutScreenshotInfo - Upon return, will hold the info for the screenshot.
+	 */
+	void GetScreenshotInfo(const FString ScreenshotName, FVirtualCameraScreenshot &OutScreenshotInfo) const;
+
+	/**
+	 * Collects a list of existing screenshot names into an out array.
+	 * @param OutArray - Upon return, will store screenshot names.
+	 */
+	void GetScreenshotNames(TArray<FString>& OutArray) const;
+
+	/** 
+	 * Change the name of a waypoint to a different name.
+	 * @param TargetWaypoint - The current name of the waypoint to be renamed
+	 * @param NewWaypointName - The new name of the waypoint
+	 * @return true if the renaming operation was successful, false if not
+	 */
+	bool RenameWaypoint(const FString& TargetWaypoint, const FString& NewWaypointName);
+
+	/**
+	 * Stores the current pawn location as a waypoint.
+	 * @return the index of the new waypoint in the saved waypoints array
+	 */
+	FString SaveWaypoint();
+
+	/**
+	 * Saves a preset into the list of presets.
+	 * @param bSaveCameraSettings - Should this preset save camera settings
+	 * @param bSaveStabilization - Should this preset save stabilization settings
+	 * @param bSaveAxisLocking - Should this preset save axis locking settings
+	 * @param bSaveMotionScale - Should this preset save motion scaled settings
+	 * @return the name of the preset
+	 */
+	FString SavePreset(const bool bSaveCameraSettings, const bool bSaveStabilization, const bool bSaveAxisLocking, const bool bSaveMotionScale);
+
+	/**
+	 * Returns a sorted TMap of the current presets.
+	 * @return a sorted TMap of settings presets
+	 */
+	TMap<FString, FVirtualCameraSettingsPreset>  GetSettingsPresets();
+
+	/**
+	 * Loads a preset using its name as a string key.
+	 * @param PresetName - The name of the preset to load
+	 * @return true if successful, false otherwise
+	 */
+	bool LoadPreset(const FString& PresetName);
+
+	/**
+	 * Deletes a preset using its name as the key.
+	 * @param PresetName - The name of the preset to delete
+	 * @return the number of values associated with the key
+	 */
+	int32 DeletePreset(const FString& PresetName);
+
+	/**
+	 * Deletes a screenshot, using its name as the key.
+	 * @param ScreenshotName - The name of the screenshot to delete
+	 * @return the number of values associated with the key
+	 */
+	int32 DeleteScreenshot(const FString& ScreenshotName);
+
+	/**
+	 * Deletes a waypoint, using its name as the key.
+	 * @param WaypointName - The name of the waypoint to delete
+	 * @return the number of values associated with the key
+	 */
+	int32 DeleteWaypoint(const FString& WaypointName);
+
+	/**
+	 * Updates the current settings to reflect those in the preset.
+	 * @param PresetToLoad - The preset containing the desired settings
+	 */
+	void UpdateSettingsFromPreset(FVirtualCameraSettingsPreset* PresetToLoad);
+
+	/**
+	 * Stores the new home location.
+	 * @param NewHomeWaypointName - The name of the new home waypoint to use; Use empty string to clear
+	 */
+	void SaveHomeWaypoint(const FString& NewHomeWaypointName);
+
+	/**
+	 * Teleports the pawn to a location associated with the specified waypoint.
+	 * @param WaypointIndex - The index of the target waypoint in the waypoints array
+	 * @return true if the teleport was successful
+	 */
+	bool TeleportToWaypoint(const FString& WaypointName);
+
+	/**
+	 * Teleports the pawn to the current marked home waypoint.
+	 * @return true if the teleport was successful
+	 */
+	bool TeleportToHomeWaypoint();
+
+	/**
+	 * Takes a screenshot from the current view and saves the location and camera settings.
+	 * @return the index of the saved screenshot in the screenshots array
+	 */
+	FString TakeScreenshot();
+
+	/**
+	 * Moves a camera to the location where a screenshot was taken and restores camera settings used for that screenshot.
+	 * @param ScreenshotIndex - The index of the screenshot in the Screenshots array
+	 * @return whether or not the load was successful
+	 */
+	bool LoadScreenshotView(const FString& ScreenshotIndex);
+
+	/**
+	 * Change the name of a waypoint to a different name.
+	 * @param TargetWaypoint - The current name of the waypoint to be renamed
+	 * @param NewWaypointName - The new name of the waypoint
+	 * @return true if the renaming operation was successful, false if not
+	 */
+	bool RenameScreenshotLocation(const FString& TargetWaypoint, const FString& NewWaypointName);
+
+	/**
+	 * Blueprint event to trigger the highlighting of a specific actor.
+	 * @param HighlightedActor - The Actor on which the focus is set
+	 */
+	UFUNCTION(BlueprintImplementableEvent)
+	void HighlightTappedActor(AActor* HighlightedActor);
+
+	/**
+	 * Blueprint event to trigger focus plane visualization for a set amount of time.
+	 */
+	UFUNCTION(BlueprintImplementableEvent)
+	void TriggerFocusPlaneTimer();
+
+	/**
+	 * Forwards any focus change commands from outside sources to the camera component.
+	 * @param NewFocusDistance - The desired focus distance for the camera
+	 */
+	void SetFocusDistance(const float NewFocusDistance) { CineCamera->SetFocusDistance(NewFocusDistance); }
+
+	/**
+	 * Forwards tracked actor changes to the camera component to be handled.
+	 * @param ActorToTrack - The actor that the focus should follow
+	 * @param TrackingPointOffset - The offset from the root of the actor to the point to focus on
+	 */
+	void SetTrackedActorForFocus(AActor* ActorToTrack, FVector TrackingPointOffset) { CineCamera->SetTrackedActorForFocus(ActorToTrack, TrackingPointOffset); }
+
+	/** 
+	 * Sets whether settings should be saved on exit. 
+	 * @param bShouldSettingsSave - Whether settings should be saved
+	 */
+	void SetSaveSettingsWhenClosing(const bool bShouldSettingsSave) { bSaveSettingsWhenClosing = bShouldSettingsSave; };
+
+	/**
+	 * Gets whether settings should save when closing.
+	 * @return whether settings should save when closing
+	 */
+	bool GetSaveSettingsWhenClosing() const { return bSaveSettingsWhenClosing; };
+
+	/**
+	 * Check to see if the camera is in autofocus mode.
+	 * @return true if camera should be autofocusing, false if not
+	 */
+	bool IsAutoFocusEnabled() const { return CineCamera->bAutoFocusEnabled; }
+
+	/**
+	 * Send any movement input data forward to the Movement Component.
+	 * @param Location - The current location of the tracker
+	 * @param Rotation - The current rotation of the tracker
+	 */
+	void ProcessMovementInput(const FVector& Location, const FRotator& Rotation) { MovementComponent->ProcessMovementInput(Location, Rotation); }
+
+	/**
+	 * Checks whether or not focus visualization can activate
+	 * @return the current state of touch event visualization
+	 */
+	bool IsFocusVisualizationAllowed() const { return bAllowFocusVisualization; }
+
+	/**
+	 * Sets whether or not focus visualization can activate
+	 * @param bShouldAllowFocusVisualization - Whether or not focus visualization should be allowed
+	 */
+	void SetAllowFocusPlaneVisualization(bool bShouldAllowFocusVisualization) { bAllowFocusVisualization = bShouldAllowFocusVisualization; }
+
+	/**
+	 * Blueprint event for signaling UI that game settings have been loaded.
+	 */
+	UFUNCTION(BlueprintImplementableEvent)
+	void LoadFinished();
+
+	/**
+	 * Set the matte aspect ratio to a new value.
+	 * @return DesiredUnits - The new unit to use for distance measures like focus distance
+	 */
+	EUnit GetDesiredDistanceUnits() const { return DesiredDistanceUnits; }
+
+	/**
+	 * Set the matte aspect ratio to a new value.
+	 * @param DesiredUnits - The new unit to use for distance measures like focus distance
+	 */
+	void SetDesiredDistanceUnits(const EUnit DesiredUnits);
+
+	/**
+	 * Sets whether or not a preset is favorited
+	 * @param PresetName - The name of the preset to adjust favorite setting for
+	 * @param bIsFavorite - Whether settings should be saved
+	 */
+	void SetPresetFavoriteStatus(const FString& PresetName, const bool bIsFavorite);
+
+	/**
+	 * Sets whether or not a screenshot is favorited
+	 * @param ScreenshotName - The name of the screenshot to adjust favorite setting for
+	 * @param bIsFavorite - Whether settings should be saved
+	 */
+	void SetScreenshotFavoriteStatus(const FString& ScreenshotName, const bool bIsFavorite);
+
+	/**
+	 * Sets whether or not a waypoint is favorited
+	 * @param WaypointName - The name of the waypoint to adjust favorite setting for
+	 * @param bIsFavorite - Whether settings should be saved
+	 */
+	void SetWaypointFavoriteStatus(const FString& WaypointName, const bool bIsFavorite);
+
+protected:
+	/** Tracks any waypoints the player has saved for teleporting */
+	TMap<FString, FVirtualCameraWaypoint> Waypoints;
+
+	/** Stores the locations of any screenshots that were taken */
+	TMap<FString, FVirtualCameraScreenshot> Screenshots;
+
+	/** Tracks which waypoint is the "home" waypoint, defaults to first value */
+	FString HomeWaypointName;
+
+	/** The next preset number */
+	static int32 PresetIndex;
+
+	/** The next waypoint number */
+	static int32 WaypointIndex;
+
+	/** The next screenshot number */
+	static int32 ScreenshotIndex;
+
+	/** The desired unit in which to display focus distance */
+	EUnit DesiredDistanceUnits;
+
+	/** Should focus plane be shown on all touch focus events */
+	UPROPERTY(BlueprintReadOnly, Category = "Default")
+	bool bAllowFocusVisualization;
+
+	/** Number formatter for Screenshots, Waypoints and Presets. */
+	uint32 MinimumIntegralDigits;
+
+	/**
+	 * Stores the current camera settings to a save file for later use.
+	 */
+	void SaveSettings();
+
+	/**
+	 * Loads camera settings from a saved file.
+	 */
+	void LoadSettings();
+
+	/**
+	 * Convenience function to leftpad numbers with zeros
+	 * @param InNumber - The input number to be padded
+	 * @param MinNumberOfCharacters - The minimum number of characters in the returned string
+	 * @return The input number padded to be at least MinNumberOfCharacters in length
+	 */
+	FString LeftPadWithZeros(int32 InNumber, int32 MinNumberOfCharacters) const;
+
+	//~ Begin IVirtualCameraPresetContainer Interface
+	virtual FString SavePreset_Implementation(const bool bSaveCameraSettings, const bool bSaveStabilization, const bool bSaveAxisLocking, const bool bSaveMotionScale) override;
+	virtual bool LoadPreset_Implementation(const FString& PresetName) override;
+	virtual int32 DeletePreset_Implementation(const FString& PresetName) override;
+	virtual TMap<FString, FVirtualCameraSettingsPreset> GetSettingsPresets_Implementation() override;
+	//~ End IVirtualCameraPresetContainer Interface
+
+	//~ Begin IVirtualCameraOptions Interface
+	virtual void SetDesiredDistanceUnits_Implementation(const EUnit DesiredUnits) override;
+	virtual EUnit GetDesiredDistanceUnits_Implementation() override;
+	virtual bool IsFocusVisualizationAllowed_Implementation() override;
+	//~ End IVirtualCameraOptions Interface
+};
diff --git a/Engine/Plugins/Experimental/VirtualCamera/Source/VirtualCamera/Public/VirtualCameraPlayerControllerBase.h b/Engine/Plugins/Experimental/VirtualCamera/Source/VirtualCamera/Public/VirtualCameraPlayerControllerBase.h
new file mode 100644
--- /dev/null
+++ b/Engine/Plugins/Experimental/VirtualCamera/Source/VirtualCamera/Public/VirtualCameraPlayerControllerBase.h
@@ -0,0 +1,928 @@
+// Copyright Epic Games, Inc. All Rights Reserved.
+
+#pragma once
+
+#include "GameFramework/PlayerController.h"
+
+#include "IVirtualCameraController.h"
+#include "LevelSequence/LevelSequencePlaybackController.h"
+#include "VirtualCameraPawnBase.h"
+
+#include "CineCameraActor.h"
+#include "Math/UnitConversion.h"
+#include "Misc/FrameNumber.h"
+#include "Misc/FrameRate.h"
+#include "Misc/Timecode.h"
+#include "VirtualCameraPlayerControllerBase.generated.h"
+
+class AVPRootActor;
+class ILiveLinkClient;
+
+UENUM(BlueprintType)
+enum class ETrackerInputSource : uint8
+{
+	/* Accelerometer data from an iPhone/iPad */
+	ARKit,
+
+	/* Takes in data from an outside source in blueprints */
+	Custom,
+
+	/* Livelink plugin tracker */
+	LiveLink
+};
+
+UENUM(BlueprintType)
+enum class ETouchInputState : uint8
+{
+	/* Allows user to select an actor to always be in focus */
+	ActorFocusTargeting,
+
+	/* Allows user to select a point on the screen to auto-focus through */
+	AutoFocusTargeting,
+
+	/* Allows the touch input to be handled in the blueprint event. This should be the default */
+	BlueprintDefined,
+
+	/* Allows for the user to focus on target on touch without exiting manual focus */
+	ManualTouchFocus,
+
+	/* Touch support for scrubbing through a sequence */
+	Scrubbing,
+
+	/* Touch and hold for attach targeting */
+	TouchAndHold,
+};
+
+UCLASS(Abstract)
+class VIRTUALCAMERA_API AVirtualCameraPlayerControllerBase : public APlayerController
+{
+	GENERATED_UCLASS_BODY()
+
+public:
+	virtual void OnPossess(APawn* InPawn) override;
+	virtual void BeginPlay() override;
+	virtual void Tick(float DeltaSeconds) override;
+	virtual void SetupInputComponent() override;
+
+	FScriptDelegate OnStop;
+
+	/** Allows user to select which tracker input should be used */
+	UPROPERTY(EditDefaultsOnly, BlueprintReadWrite, BlueprintSetter="SetInputSource", Category = "Camera Settings")
+	ETrackerInputSource InputSource;
+
+	/** Controller for level sequence playback */
+	UPROPERTY(EditAnywhere, BlueprintReadOnly, Category = "Camera Settings")
+	FName LiveLinkTargetName;
+
+	/** Offset applied to calculated location before tracker transform is added */
+	UPROPERTY(EditAnywhere, BlueprintReadOnly, Category = "Camera Settings")
+	FTrackingOffset TrackerPreOffset;
+
+	/** Offset applied to calculated location after tracker transform is added */
+	UPROPERTY(EditAnywhere, BlueprintReadOnly, Category = "Camera Settings")
+	FTrackingOffset TrackerPostOffset;
+
+	/** Class of CameraActor to spawn to allow user to use their own customized camera */
+	UPROPERTY(EditAnywhere, BlueprintReadOnly, NoClear, Category = "Camera Settings")
+	TSubclassOf<ACineCameraActor> TargetCameraActorClass;
+
+	/** Array of any properties that should be recorded */
+	UPROPERTY(EditAnywhere, Category = "Recording")
+	TArray<FName> RequiredSequencerRecorderCameraSettings;
+
+	UPROPERTY(BlueprintAssignable, Category = "Movement")
+	FVirtualCameraResetOffsetsDelegate OnOffsetReset;
+
+	UPROPERTY(transient, BlueprintReadOnly, Category = "VirtualCamera")
+	TObjectPtr<AVPRootActor> RootActor;
+
+	/**
+	 * Overridable function to allow user to get tracker data from blueprints.
+	 * @param OutTrackerLocation - The current location of the tracker being used
+	 * @param OutTrackerRotation - The current rotation of the tracker being used
+	 */
+	UFUNCTION(BlueprintImplementableEvent, Category = "VirtualCamera")
+	void GetCustomTrackerLocationAndRotation(FVector& OutTrackerLocation, FRotator& OutTrackerRotation);
+
+	/**
+	 * Blueprint Event for updating position of autofocus reticle.
+	 * @param NewReticleLocation - The new location of the autofocus reticle
+	 */
+	UFUNCTION(BlueprintImplementableEvent, Category = "VirtualCamera")
+	void UpdateFocusReticle(FVector NewReticleLocation);
+
+	/**
+	 * Blueprint Event for updating if a sequence can be recorded or not.
+	 * @param bIsRecordEnabled - If the loaded sequence can record or not.
+	 */
+	UFUNCTION(BlueprintImplementableEvent, Category = "VirtualCamera")
+	void OnRecordEnabledStateChanged(bool bIsRecordEnabled);
+
+	/**
+	 * Blueprint Event for when a sequence stops playing.
+	 */
+	UFUNCTION(BlueprintImplementableEvent, Category = "VirtualCamera")
+	void OnStopped();
+
+	/**
+	 * Sets the autofocus point to the correct starting location at the center of the screen
+	 */
+	UFUNCTION(BlueprintCallable, Category = "VirtualCamera")
+	void InitializeAutoFocusPoint();
+
+	/**
+	 * Returns the target camera that was spawned for this play
+	 */
+	UFUNCTION(BlueprintCallable, Category = "VirtualCamera")
+	ACineCameraActor* GetTargetCamera();
+
+	/**
+	 * Set the input source
+	 */
+	UFUNCTION(BlueprintSetter, Category = "Camera Settings")
+	void SetInputSource(ETrackerInputSource InInputSource);
+
+	/**
+	 * Blueprint event for when the focus method is changed.
+	 * @param NewFocuMethod - The focus method
+	 */
+	UFUNCTION(BlueprintImplementableEvent, Category = "VirtualCamera")
+	void FocusMethodChanged(EVirtualCameraFocusMethod NewFocusMethod);
+
+	/**
+	 * Converts a distance in unreal units (cm) to other units of measurement for display.
+	 * @param InputDistance - A distance in Unreal Units to be converted
+	 * @param ConversionUnit - The desired unit of distance to be used
+	 * @return the distance converted and put into a string for display
+	 */
+	UFUNCTION(BlueprintPure, Category = "VirtualCamera")
+	FString GetDistanceInDesiredUnits(const float InputDistance, const EUnit ConversionUnit) const;
+
+	/**
+	 * Function to handle delegate when form playback controller. This notifies whether a loaded level sequence can record.
+	 * @param bIsRecordEnabled - If this sequence can record or not.
+	 */
+	void HandleRecordEnabledStateChange(const bool bIsRecordEnabled);
+
+protected:
+	/** Client interface to the LiveLink plugin for gathering data on an actor */
+	ILiveLinkClient* LiveLinkClient;
+
+	/** The 2D vector corresponding to which point on the screen will be used for autofocus */
+	FVector2D AutoFocusScreenPosition;
+
+	/** The current focus method in use */
+	EVirtualCameraFocusMethod CurrentFocusMethod;
+
+	/** The touch input state to determine the context of how touch input should be handled. */
+	ETouchInputState TouchInputState;
+
+	/** The previous touch input state. */
+	ETouchInputState PreviousTouchInput;
+
+	/** Controller for level sequence playback */
+	UPROPERTY(Transient)
+	TObjectPtr<ULevelSequencePlaybackController> LevelSequencePlaybackController;
+	
+	/** Target camera that is spawned or possessed on begin play for the sequence controller */
+	UPROPERTY(Transient)
+	TObjectPtr<ACineCameraActor> TargetCameraActor;
+
+	/** Cached value for IsVirtualCameraControlledByRemoteSession() */
+	UPROPERTY(Transient, BlueprintReadOnly, Category = "VirtualCamera")
+	bool bCachedIsVirtualCameraControlledByRemoteSession;
+
+	/** Cached value for IsVirtualCameraControlledByRemoteSession() */
+	UPROPERTY(Transient, BlueprintReadOnly, Category = "VirtualCamera")
+	bool bCachedShouldUpdateTargetCameraTransform;
+	
+
+protected:
+
+	/**
+	 * Get the current tracker location and rotation based on selected input method.
+	 * @param OutTrackerLocation - The current location of the tracker being used
+	 * @param OutTrackerRotation - The current rotation of the tracker being used
+	 * @return true if getting the location/rotation was successful, false otherwise
+	 */
+	bool GetCurrentTrackerLocationAndRotation(FVector& OutTrackerLocation, FRotator& OutTrackerRotation);
+
+	/**
+	 * Sets the focus distance to an object selected by the player based on touch input.
+	 * @param - The index of the touch input that occurred
+	 * @param - the location of the touch in screen space
+	 */
+	void SetFocusDistanceToActor(const ETouchIndex::Type TouchIndex, const FVector& Location);
+
+	/**
+	 * Gets the current Virtual Camera Pawn
+	 * @return pointer to the current possessed pawn if it's a Virtual Camera Pawn
+	 */
+	AVirtualCameraPawnBase* GetVirtualCameraPawn() const;
+
+	/**
+	 * Sets the focus distance through a point on the screen.
+	 * @param ScreenPosition - The point on the screen to trace through
+	 */
+	void SetFocusDistanceThroughPoint(const FVector2D ScreenPosition);
+
+	/**
+	 * Convenience function to get camera component for UI functions.
+	 * @return pointer to camera component if one is found, null otherwise
+	 */
+	UVirtualCameraCineCameraComponent* GetVirtualCameraCineCameraComponent() const;
+
+	/**
+	 * Convenience function to get movement component for UI functions
+	 * @return pointer to movement component if one is found, null otherwise
+	 */
+	UVirtualCameraMovementComponent* GetVirtualCameraMovementComponent() const;
+
+	/**
+	 * Pilot the controlled camera during recording, copying over settings from the pawn.
+	 */
+	void PilotTargetedCamera(AVirtualCameraPawnBase* PawnToFollow, UVirtualCameraCineCameraComponent* CameraToFollow);
+
+	/**
+	 * Override of InputTouch, used to handle touch and hold events.
+	 */
+	virtual bool InputTouch(uint32 Handle, ETouchType::Type Type, const FVector2D& TouchLocation, float Force, FDateTime DeviceTimestamp, uint32 TouchpadIndex) override;
+
+	/**
+	 * Determine how to handle touch input based on current context.
+	 * @param TouchIndex - The touch index that triggered the input
+	 * @param Location - The location of the touch input in screen space
+	 */
+	void OnTouchInput(const ETouchIndex::Type TouchIndex, const FVector Location);
+
+	/**
+	 * Moves the point through which the camera auto focuses.
+	 * @param TouchInded - The touch index that triggered the input
+	 * @paran Location - The location of touch input, and new location of auto focus
+	 */
+	void UpdateScreenFocus(const ETouchIndex::Type TouchIndex, const FVector Location);
+
+	/**
+	 * Handle moving forward/backward input from a controller or touch interface.
+	 * @param InValue - The axis input value
+	 */
+	void OnMoveForward(const float InValue);
+
+	/**
+	 * Handle moving left/right input from a controller or touch interface.
+	 * @param InValue - The axis input value
+	 */
+	void OnMoveRight(const float InValue);
+
+	/**
+	 * Handle moving up/down input from a controller or touch interface.
+	 * @param InValue - The axis input value
+	 */
+	void OnMoveUp(const float InValue);
+
+	/**
+	 * Will trigger the event to show focus visualization, unless disabled by user.
+	 */
+	void ShowFocusPlaneFromTouch();
+
+	/**
+	 * Get the current tracker data and update the movement component.
+	 */
+	void UpdatePawnWithTrackerData();
+
+	/**
+	 * Checks of the touch input
+	 */
+	bool IsLocationWithinMatte(const FVector Location) const;
+
+	/**
+	 * Broadcast offset resets when movement component broadcasts
+	 */
+	UFUNCTION()
+	void BroadcastOffsetReset() { OnOffsetReset.Broadcast(); }
+
+	/*
+	 * Is this machine should display the Virtual Camera UI and establish a connection with the remote session app.
+	 */
+	UFUNCTION(BlueprintNativeEvent, Category = "VirtualCamera")
+	bool IsVirtualCameraControlledByRemoteSession() const;
+
+	/*
+	 * In multi user session, how should we update the information across different sessions.
+	 */
+	UFUNCTION(BlueprintNativeEvent, Category = "VirtualCamera")
+	bool ShouldUpdateTargetCameraTransform() const;
+
+	/***** UI Interface *****/
+public:
+
+	/**
+	 * Adds a blendable object to the camera's post process settings.
+	 * @param BlendableToAdd - The blendable that will be added to the post process settings
+	 * @param Weight - The weighting of the blendable's alpha
+	 */
+	UFUNCTION(BlueprintCallable, Category = "Virtual Camera|Cinematic Camera")
+	void AddBlendableToCamera(TScriptInterface<IBlendableInterface> BlendableToAdd, float Weight);
+
+	/**
+	 * Adjusts the aperture of the camera to an adjacent preset value.
+	 * @param bShiftUp - If true, the aperture will increase; if false, it will decrease
+	 * @return the new aperture of the camera
+	 */
+	UFUNCTION(BlueprintCallable, Category = "Virtual Camera|Cinematic Camera|Aperture")
+	float ChangeAperturePreset(const bool bShiftUp);
+
+	/**
+	 * Adjust the focal length of the camera to an adjecent preset value.
+	 * @param bShiftUp - If true, the focal length will increase; if false, it will decrease
+	 * @return the new focal length of the camera in mm
+	 */
+	UFUNCTION(BlueprintCallable, Category = "Virtual Camera|Cinematic Camera|Focal Length")
+	float ChangeFocalLengthPreset(const bool bShiftUp);
+
+	/**
+	 * Clears the current level sequence player, needed when recording clean takes of something
+	 */
+	UFUNCTION(BlueprintCallable, Category = "Virtual Camera|Sequencer|Playback")
+	void ClearActiveLevelSequence();
+
+	/**
+	 * Deletes a preset using its name as the key.
+	 * @param PresetName - The name of the preset to delete
+	 * @return the number of values associated with the key
+	 */
+	UFUNCTION(BlueprintCallable, Category = "Virtual Camera|Settings")
+	int32 DeletePreset(FString PresetName);
+
+	/**
+	 * Deletes a screenshot, using its name as the key.
+	 * @param ScreenshotName - The name of the screenshot to delete
+	 * @return the number of values associated with the key
+	 */
+	UFUNCTION(BlueprintCallable, Category = "Virtual Camera|Screenshots")
+	int32 DeleteScreenshot(FString ScreenshotName);
+
+	/**
+	 * Deletes a waypoint, using its name as the key.
+	 * @param WaypointName - The name of the waypoint to delete
+	 * @return the number of values associated with the key
+	 */
+	UFUNCTION(BlueprintCallable, Category = "Virtual Camera|Waypoints")
+	int32 DeleteWaypoint(FString WaypointName);
+
+	/**
+	 * Returns the asset name of the currently selected sequence
+	 * @return the name of the crrent selected sequence; returns empty string if no selected sequence
+ 	 */
+	UFUNCTION(BlueprintPure, Category = "Virtual Camera|Sequencer|Playback")
+	FString GetActiveLevelSequenceName();
+
+	/**
+	 * Returns the currently selected sequence
+	 * @return the current selected sequence; returns nullptr if no selected sequence
+	 */
+	UFUNCTION(BlueprintPure, Category = "Virtual Camera|Sequencer|Playback")
+	ULevelSequence* GetActiveLevelSequence();
+	
+	/**
+	 * Gets stabalization scale for a specific axis.
+	 * @param AxisToRetrieve - The axis of the stabilization value needed
+	 * @return the stabilization scale for the given axis
+	 */
+	UFUNCTION(BlueprintPure, Category = "Virtual Camera|Movement")
+	float GetAxisStabilizationScale(EVirtualCameraAxis AxisToRetrieve);
+
+	/**
+	* Gets movement scale for a specific axis.
+	* @param AxisToRetrieve - The axis of the scale value needed
+	* @return the movement scale for the given axis
+	*/
+	UFUNCTION(BlueprintPure, Category = "Virtual Camera|Movement")
+	float GetAxisMovementScale(EVirtualCameraAxis AxisToRetrieve);
+
+	/**
+	 * Get the current aperture value on the camera.
+	 * @return the current aperture of the camera component
+	 */
+	UFUNCTION(BlueprintPure, Category = "Virtual Camera|Cinematic Camera|Aperture")
+	float GetCurrentAperture();
+
+	/**
+	 * Get the name of the current preset filmback option on the camera.
+	 * @return the name of the current preset filmback option
+	 */
+	UFUNCTION(BlueprintPure, Category = "Virtual Camera|Cinematic Camera|Filmback")
+	FString GetCurrentFilmbackName();
+
+	/**
+	 * Get the current focal length value on the camera.
+	 * @return the current focal length of the camera being viewed
+	 */
+	UFUNCTION(BlueprintPure, Category = "Virtual Camera|Cinematic Camera|Focal Length")
+	float GetCurrentFocalLength();
+
+	/**
+	 * Returns the current focus distance of the camera.
+	 * @return the focal distance of the camera
+	 */
+	UFUNCTION(BlueprintPure, Category = "Virtual Camera|Cinematic Camera|Focus")
+	float GetCurrentFocusDistance();
+
+	/**
+	* Returns the current focus method.
+	* @return the focus method being used
+	*/
+	UFUNCTION(BlueprintPure, Category = "Virtual Camera|Cinematic Camera|Focus")
+	EVirtualCameraFocusMethod GetCurrentFocusMethod() { return CurrentFocusMethod; }
+
+	/**
+	 * Get the end position of the currently selected sequence
+	 * @return the end position of the sequence in FrameNumber
+	 */
+	UFUNCTION(BlueprintPure, Category = "Virtual Camera|Sequencer|Playback")
+	FFrameNumber GetCurrentSequencePlaybackEnd();
+
+	/**
+	 * Get the start position of the currently selected sequence
+	 * @return the start position of the sequence in FrameNumber
+	 */
+	UFUNCTION(BlueprintPure, Category = "Virtual Camera|Sequencer|Playback")
+	FFrameNumber GetCurrentSequencePlaybackStart();
+
+	/**
+	 * Gets the locked to camera cut from the active LevelSequence
+	 */
+	UFUNCTION(BlueprintPure, Category = "Virtual Camera|Sequencer|Playback")
+	bool IsSequencerLockedToCameraCut();
+
+	/**
+	 * Sets the locked to camera cut from the active LevelSequence
+	 */
+	UFUNCTION(BlueprintCallable, Category = "Virtual Camera|Sequencer|Playback")
+	void SetSequencerLockedToCameraCut(bool bLockView);
+
+	/**
+	 * Get the frame rate of the currently selected sequence
+	 * @return the frame rate in frames per second
+	 */
+	UFUNCTION(BlueprintPure, Category = "Virtual Camera|Sequencer|Playback")
+	FFrameRate GetCurrentSequenceFrameRate();
+
+	/**
+	 * Returns previously set unit of distance.
+	 * @return DesiredUnits - The unit that is used for distance measures like focus distance
+	 */
+	UFUNCTION(BlueprintPure, Category = "Virtual Camera|Settings")
+	EUnit GetDesiredDistanceUnits();
+
+	/**
+	 * Get the current color of the focus plane that should be used 
+	 * @return - The name of the current or next sequence that will be recorded
+	 */
+	UFUNCTION(BlueprintPure, Category = "Virtual Camera|Cinematic Camera|Focus")
+	FColor GetFocusPlaneColor();
+
+	/**
+	 * Stores the names of all available filmback presets into an array.
+	 * @param OutFilmbackPresets - Upon return, will contain all available filmback presets
+	 * @return true if operation was successful
+	 */
+	UFUNCTION(BlueprintPure, BlueprintCallable, Category = "Virtual Camera|Cinematic Camera|Filmback")
+	bool GetFilmbackPresetOptions(TArray<FString>& OutFilmbackPresetsArray);
+
+	/**
+	 * Gets the length of the currently selected level sequence
+	 * @return the length of the level sequence in FrameNumber
+	 */
+	UFUNCTION(BlueprintPure, Category = "Virtual Camera|Sequencer|Playback")
+	FFrameNumber GetLevelSequenceLength();
+
+	/**
+	 * Returns the names of each level sequence asset in the project
+	 * @param OutLevelSequenceNames - Upon return, array will contain all available level sequence names
+	 */
+	UFUNCTION(BlueprintCallable, Category = "Virtual Camera|Sequencer|Playback")
+	void GetLevelSequences(TArray<FLevelSequenceData>& OutLevelSequenceNames);
+
+	/**
+	 * Returns the current matte aspect ratio.
+	 * @return the current matte setting
+	 */
+	UFUNCTION(BlueprintPure, Category = "Virtual Camera|Cinematic Camera|Filmback")
+	float GetMatteAspectRatio();
+
+	/**
+	 * Set the matte aspect ratio to a new value.
+	 * @return The desired matte opacity to use
+	 */
+	UFUNCTION(BlueprintPure, Category = "Virtual Camera|Cinematic Camera|Filmback")
+	float GetMatteOpacity();
+
+	/**
+	 * Returns the values of all matte options.
+	 * @param OutMatteValues - Upon return, array will contain all matte values.
+	 */
+	UFUNCTION(BlueprintPure, Category = "Virtual Camera|Cinematic Camera|Filmback")
+	void GetMatteValues(TArray<float>& OutMatteValues);
+
+	/**
+	 * Gets the playback position of the level sequence
+	 * @return the current playback position
+	 */
+	UFUNCTION(BlueprintPure, Category = "Virtual Camera|Sequencer|Playback")
+	FFrameTime GetPlaybackPosition();
+
+	/**
+	 * Gets the playback Timecode of the level sequence
+	 * @return the current playback Timecode
+	 */
+	UFUNCTION(BlueprintPure, Category = "Virtual Camera|Sequencer|Playback")
+	FTimecode GetPlaybackTimecode();
+
+	/**
+	 * Returns the information associated with a Screenshot.
+	 * @param ScreenshotName - The name of the screenshot to retrieve
+	 * @param OutScreenshotInfo - Upon return, will hold the info for the screenshot.
+	 */
+	UFUNCTION(BlueprintPure, Category = "Virtual Camera|Screenshots")
+	void GetScreenshotInfo(FString ScreenshotName, FVirtualCameraScreenshot& OutScreenshotInfo);
+
+	/**
+	 * Collects a list of existing screenshot names.
+	 * @param OutArray - Upon return, will store screenshot names.
+	 */
+	UFUNCTION(BlueprintPure, Category = "Virtual Camera|Screenshots")
+	void GetScreenshotNames(TArray<FString>& OutArray);
+
+	/**
+	 * Returns a sorted TMap of the current presets.
+	 * @return a sorted TMap of settings presets
+	 */
+	UFUNCTION(BlueprintPure, Category = "Virtual Camera|Settings")
+	TMap<FString, FVirtualCameraSettingsPreset> GetSettingsPresets();
+
+	/**
+	 * Sets the current state of touch input.
+	 * @return the current state of the input
+	 */
+	UFUNCTION(BlueprintPure, Category = "Virtual Camera|Input")
+	ETouchInputState GetTouchInputState() { return TouchInputState; }
+
+	/**
+	 * Returns the information associated with a waypoint.
+	 * @param WaypointName - The name of the waypoint to retrieve
+	 * @param OutWaypointInfo - Upon return, will hold the info for that waypoint.
+	 */
+	UFUNCTION(BlueprintPure, Category = "Virtual Camera|Waypoints")
+	void GetWaypointInfo(FString WaypointName, FVirtualCameraWaypoint &OutWaypointInfo);
+
+	/**
+	 * Collects a list of existing waypoint names.
+	 * @param OutArray - Upon return, will store the waypoint names.
+	 */
+	UFUNCTION(BlueprintPure, Category = "Virtual Camera|Waypoints")
+	void GetWaypointNames(TArray<FString>& OutArray);
+
+	/**
+	 * Checks if an axis is locked.
+	 * @param AxisToCheck - The axis being checked
+	 * @return true if locked false otherwise
+	 */
+	UFUNCTION(BlueprintPure, Category = "Virtual Camera|Movement")
+	bool IsAxisLocked(EVirtualCameraAxis AxisToCheck);
+
+	/**
+	 * Checks whether or not focus visualization can activate
+	 * @return the current state of touch event visualization
+	 */
+	UFUNCTION(BlueprintPure, Category = "Virtual Camera|Cinematic Camera|Focus")
+	bool IsFocusVisualizationAllowed();
+
+	/**
+	 * Check to see if the sequence is playing
+	 * @return if the sequence is playing
+	 */
+	UFUNCTION(BlueprintPure, Category = "Virtual Camera|Sequencer|Playback")
+	bool IsPlaying();
+
+	/**
+	 * Helper to check if touch input state is in a touch focus mode.
+	 * @return if touch input is in touch focus mode
+	 */
+	UFUNCTION(BlueprintPure, Category = "Virtual Camera|Input")
+	bool IsTouchInputInFocusMode();
+
+	/**
+	 * Get whether or not global boom is being used when navigating with the joysticks
+	 * @param bShouldUseGlobalBoom - True if using global Z axis, false if using local Z axis
+	 */
+	UFUNCTION(BlueprintPure, Category = "Virtual Camera|Movement")
+	bool IsUsingGlobalBoom();
+
+	/**
+	 * Goes to the end of the level sequence and pauses
+	 */
+	UFUNCTION(BlueprintCallable, Category = "Virtual Camera|Sequencer|Playback")
+	void JumpToLevelSequenceEnd();
+
+	/**
+	 * Goes to the start of the level sequence and pauses
+	 */
+	UFUNCTION(BlueprintCallable, Category = "Virtual Camera|Sequencer|Playback")
+	void JumpToLevelSequenceStart();
+
+	/**
+	 * Sets the playback position of the level sequence.
+	 * @param InFrameNumber - New FrameNumber to jump to.
+	 */
+	UFUNCTION(BlueprintCallable, Category = "Virtual Camera|Sequencer|Playback")
+	void JumpToPlaybackPosition(const FFrameNumber& InFrameNumber);
+
+	/**
+	 * Loads a preset using its name as a string key.
+	 * @returns true if successful, false otherwise
+	 */
+	UFUNCTION(BlueprintCallable, Category = "Virtual Camera|Settings")
+	bool LoadPreset(FString PresetName);
+
+	/**
+	 * Moves a camera to the location where a screenshot was taken and restores camera settings used for that screenshot.
+	 * @param ScreenshotIndex - The index of the screenshot in the Screenshots array
+	 * @return whether or not the load was successful
+	 */
+	UFUNCTION(BlueprintCallable, Category = "Virtual Camera|Screenshots")
+	bool LoadScreenshotView(const FString& ScreenshotName);
+
+	/**
+	 * Pauses the playing of the current level sequence.
+	 */
+	UFUNCTION(BlueprintCallable, Category = "Virtual Camera|Sequencer|Playback")
+	void PauseLevelSequence();
+
+	/**
+	 * Plays current level sequence
+	 */
+	UFUNCTION(BlueprintCallable, Category = "Virtual Camera|Sequencer|Playback")
+	void PlayLevelSequence();
+
+	/**
+	 * Plays current level sequence in reverse
+	 */
+	UFUNCTION(BlueprintCallable, Category = "Virtual Camera|Sequencer|Playback")
+	void PlayLevelSequenceInReverse();
+
+	/**
+	 * Moves the camera back to actor root and aligns rotation with the input tracker.
+	 */
+	UFUNCTION(BlueprintCallable, Category = "Virtual Camera|Movement")
+	void ResetCameraOffsetsToTracker();
+
+	/**
+	 * Plays current level sequence from the current time.
+	 */
+	UFUNCTION(BlueprintCallable, Category = "Virtual Camera|Sequencer|Playback")
+	void ResumeLevelSequencePlay();
+
+	/**
+	 * Stores the new home waypoint location.
+	 * @param NewHomeWaypointName - The name of the new home waypoint to use
+	 */
+	UFUNCTION(BlueprintCallable, Category = "Virtual Camera|Waypoints")
+	void SaveHomeWaypoint(const FString& NewHomeWaypointName);
+
+	/**
+	 * Saves a preset into the list of presets.
+	 * @param bSaveCameraSettings - Should this preset save camera settings
+	 * @param bSaveStabilization - Should this preset save stabilization settings
+	 * @param bSaveAxisLocking - Should this preset save axis locking settings
+	 * @param bSaveMotionScale - Should this preset save motion scaled settings
+	 * @return the name of the preset
+	 */
+	UFUNCTION(BlueprintCallable, Category = "Virtual Camera|Settings")
+	FString SavePreset(bool bSaveCameraSettings, bool bSaveStabilization, bool bSaveAxisLocking, bool bSaveMotionScale);
+
+	/**
+	 * Stores the current pawn location as a waypoint.
+	 * @return the index of the new waypoint in the saved waypoints array
+	 */
+	UFUNCTION(BlueprintCallable, Category = "Virtual Camera|Waypoints")
+	FString SaveWaypoint();
+
+	/**
+	 * Changes the active level sequence to a new level sequence.
+	 * @param LevelSequenceName - The name of the level sequence to select
+	 * @return true if a valid level sequence player was found, false if no level sequence player is currently available
+	 */
+	UFUNCTION(BlueprintCallable, Category = "Virtual Camera|Sequencer|Playback")
+	bool SetActiveLevelSequence(ULevelSequence* InNewLevelSequence);
+
+	/**
+	 * Sets whether or not to use focus visualization
+	 * @param bShouldAllowFocusVisualization - the current state of touch event visualization
+	 */
+	UFUNCTION(BlueprintCallable, Category = "Virtual Camera|Cinematic Camera|Focus")
+	void SetAllowFocusPlaneVisualization(bool bShouldAllowFocusVisualization);
+
+	/**
+	 * Sets the stabilization rate for a given lock.
+	 * @param AxisToAdjust - The axis whose stabilization rate should be changed
+	 * @param NewStabilizationAmount - The stabilization amount we should attempt to set the value
+	 * @return the actual value the stabilization amount was set to after clamping
+	 */
+	UFUNCTION(BlueprintCallable, Category = "Virtual Camera|Movement")
+	float SetAxisStabilizationScale(EVirtualCameraAxis AxisToAdjust, float NewStabilizationAmount);
+
+	/**
+	 * Set the current aperture value on the camera.
+	 * @param NewAperture - The current aperture of the camera component
+	 */
+	UFUNCTION(BlueprintCallable, Category = "Virtual Camera|Cinematic Camera|Aperture")
+	void SetCurrentAperture(float NewAperture);
+
+	/**
+	 * Sets the current focal length of the cinematic camera to a given value.
+	 * @param NewFocalLength - The target focal length of the camera
+	 */
+	UFUNCTION(BlueprintCallable, Category = "Virtual Camera|Cinematic Camera|Focal Length")
+	void SetCurrentFocalLength(const float NewFocalLength);
+
+	/**
+	 * Sets the current focus distance of the cinematic camera to a given value.
+	 * @param NewFocusDistance - The target focus distance of the camera
+	 */
+	UFUNCTION(BlueprintCallable, Category = "Virtual Camera|Cinematic Camera|Focus")
+	void SetCurrentFocusDistance(const float NewFocusDistance);
+
+	/**
+	 * Sets unit of distance.
+	 * @param DesiredUnits - The new unit to use for distance measures like focus distance
+	 */
+	UFUNCTION(BlueprintCallable, Category = "Virtual Camera|Settings")
+	void SetDesiredDistanceUnits(const EUnit DesiredUnits);
+
+	/**
+	 * Set the filmback settings to a new filmback preset.
+	 * @return true if operation was successful; false if NewFilmbackPreset is not a valid option
+	 */
+	UFUNCTION(BlueprintCallable, Category = "Virtual Camera|Cinematic Camera|Filmback")
+	bool SetFilmbackPresetOption(const FString& NewFilmbackPreset);
+
+	/**
+	 * Sets the camera focus method.
+	 * @param NewFocusMethod - The new focus method to be used by the camera
+	 */
+	UFUNCTION(BlueprintCallable, Category = "Virtual Camera|Cinematic Camera|Focus")
+	void SetFocusMethod(const EVirtualCameraFocusMethod NewFocusMethod);
+
+	/**
+	 * Changes focus plane color.
+	 * @param - The new color for the focus plane when focus Visualization is activate
+	 */
+	UFUNCTION(BlueprintCallable, Category = "Virtual Camera|Cinematic Camera|Focus")
+	void SetFocusPlaneColor(const FColor NewFocusPlaneColor);
+
+	/**
+	 * Toggles focus visualization tools on camera.
+	 * @param bShowFocusVisualization - The desired state of the visualization tools
+	 */
+	UFUNCTION(BlueprintCallable, Category = "Virtual Camera|Cinematic Camera|Focus")
+	void SetFocusVisualization(bool bShowFocusVisualization);
+
+	/**
+	 * Checks whether or not focus visualization is activate
+	 * @return the current state of touch event visualization
+	 */
+	UFUNCTION(BlueprintPure, Category = "Virtual Camera|Cinematic Camera|Focus")
+	bool IsFocusVisualizationActivated() const;
+
+	/**
+	 * Set the matte aspect ratio to a new value.
+	 * @param NewMatteAspectRatio - The desired matte aspect to use
+	 * @return true if operation was successful; false if NewMattePreset is not a valid input
+	 */
+	UFUNCTION(BlueprintCallable, Category = "Virtual Camera|Cinematic Camera|Filmback")
+	bool SetMatteAspectRatio(const float NewMatteAspectRatio);
+
+	/**
+	 * Set the matte aspect ratio to a new value.
+	 * @param NewMatteOpacity - The desired matte aspect opacity to use
+	 */
+	UFUNCTION(BlueprintCallable, Category = "Virtual Camera|Cinematic Camera|Filmback")
+	void SetMatteOpacity(const float NewMatteOpacity);
+
+	/**
+	 * Sets the movement scale of the camera actor.
+	 * @param AxisToAdjust - The axis to set movement scale on (setting for rotation axes is allowed but has no effect)
+	 * @param NewMovementScale - The desired scaling factor for calculating movement along this axis
+	 */
+	UFUNCTION(BlueprintCallable, Category = "Virtual Camera|Movement")
+	void SetMovementScale(const EVirtualCameraAxis AxisToAdjust, const float NewMovementScale);
+
+	/**
+	 * Sets whether or not a preset is favorited
+	 * @param PresetName - The name of the preset to adjust favorite setting for
+	 * @param bIsFavorite - Whether settings should be saved
+	 */
+	UFUNCTION(BlueprintCallable, Category = "Virtual Camera|Screenshots")
+	void SetPresetFavoriteStatus(const FString& PresetName, const bool bIsFavorite);
+
+	/**
+	 * Sets whether settings should be saved on exit.
+	 * @param bShouldSettingsSave - Whether settings should be saved
+	 */
+	UFUNCTION(BlueprintCallable, Category = "Virtual Camera|Settings")
+	void SetSaveSettingsWhenClosing(bool bShouldSettingsSave);
+
+	/**
+	 * Sets whether or not a screenshot is favorited
+	 * @param ScreenshotName - The name of the screenshot to adjust favorite setting for
+	 * @param bIsFavorite - Whether settings should be saved
+	 */
+	UFUNCTION(BlueprintCallable, Category = "Virtual Camera|Screenshots")
+	void SetScreenshotFavoriteStatus(const FString& ScreenshotName, const bool bIsFavorite);
+
+	/**
+	 * Sets the current state of touch input.
+	 * @param NewInputState - The new state of the input
+	 */
+	UFUNCTION(BlueprintCallable, Category = "Virtual Camera|Input")
+	void SetTouchInputState(ETouchInputState NewInputState) { TouchInputState = NewInputState; }
+
+	/**
+	 * Sets whether or not global boom should be used when navigating with the joysticks
+	 * @param bShouldUseGlobalBoom - True to use global Z axis, false to use local Z axis
+	 */
+	UFUNCTION(BlueprintCallable, Category = "Virtual Camera|Movement")
+	void SetUseGlobalBoom(bool bShouldUseGlobalBoom);
+
+	/**
+	 * Sets whether or not a waypoint is favorited
+	 * @param WaypointName - The name of the waypoint to adjust favorite setting for
+	 * @param bIsFavorite - Whether settings should be saved
+	 */
+	UFUNCTION(BlueprintCallable, Category = "Virtual Camera|Waypoints")
+	void SetWaypointFavoriteStatus(const FString& WaypointName, const bool bIsFavorite);
+
+	/**
+	 * Set the value for the option to zero out dutch when locking that axis.
+	 * @param bInValue - The new setting to use for zeroing dutch on lock
+	 */
+	UFUNCTION(BlueprintCallable, Category = "Virtual Camera|Movement")
+	void SetZeroDutchOnLock(const bool bInValue);
+
+	/**
+	 * Check whether settings should save when closing
+	 * @return whether settings should save when closing
+	 */
+	UFUNCTION(BlueprintPure, Category = "Virtual Camera|Settings")
+	bool ShouldSaveSettingsWhenClosing();
+
+	/**
+	 * Stops the currently playing level sequence.
+	 */
+	UFUNCTION(BlueprintCallable, Category = "Virtual Camera|Sequencer|Playback")
+	void StopLevelSequencePlay();
+
+	/**
+	 * Takes a screenshot from the current view and saves the location and camera settings.
+	 * @return the index of the saved screenshot in the screenshots array
+	 */
+	UFUNCTION(BlueprintCallable, Category = "Virtual Camera|Screenshots")
+	FString TakeScreenshot();
+
+	/**
+	 * Teleports to the current marked home waypoint.
+	 */
+	UFUNCTION(BlueprintCallable, Category = "Virtual Camera|Waypoints")
+	void TeleportToHomeWaypoint();
+
+	/**
+	 * Teleports the pawn to a specific location.
+	 * @param WaypointIndex - The index of the target waypoint in the waypoints array
+	 * @return if the teleport was successful, return true
+	 */
+	UFUNCTION(BlueprintCallable, Category = "Virtual Camera|Waypoints")
+	bool TeleportToWaypoint(const FString& WaypointIndex);
+
+	/**
+	 * Toggles the freeze on a given axis; returns new frozen state.
+	 * @param AxisToToggle - The axis whose lock should be toggled
+	 * @return the new frozen state of AxisToToggle (true = frozen)
+	 */
+	UFUNCTION(BlueprintCallable, Category = "Virtual Camera|Movement")
+	bool ToggleAxisFreeze(const EVirtualCameraAxis AxisToToggle);
+
+	/**
+	 * Toggles the lock on a given axis; returns new locked state.
+	 * @param AxisToToggle - The axis whose lock should be turned on/off
+	 * @return the new locked state of AxisToToggle (true = locked)
+	 */
+	UFUNCTION(BlueprintCallable, Category = "Virtual Camera|Movement")
+	bool ToggleAxisLock(const EVirtualCameraAxis AxisToToggle);
+
+	/**
+	 * Helper to bring the game viewport to front.
+	 */
+	UFUNCTION(BlueprintCallable, Category = "Virtual Camera")
+	void ActivateGameViewport();
+};
diff --git a/Engine/Plugins/Experimental/VirtualCamera/Source/VirtualCamera/Public/VirtualCameraSaveGame.h b/Engine/Plugins/Experimental/VirtualCamera/Source/VirtualCamera/Public/VirtualCameraSaveGame.h
new file mode 100644
--- /dev/null
+++ b/Engine/Plugins/Experimental/VirtualCamera/Source/VirtualCamera/Public/VirtualCameraSaveGame.h
@@ -0,0 +1,347 @@
+// Copyright Epic Games, Inc. All Rights Reserved.
+
+#pragma once
+
+#include "GameFramework/SaveGame.h"
+#include "Math/UnitConversion.h"
+#include "VirtualCameraSaveGame.generated.h"
+
+UENUM()
+enum class EVirtualCameraAxis : uint8
+{
+	/* Locking for x location movement; Truck */
+	LocationX,
+
+	/* Locking for y location movement; Dolly */
+	LocationY,
+
+	/* Locking for z location movement; Boom */
+	LocationZ,
+
+	/* Locking for x rotation movement; Dutch or Roll */
+	RotationX,
+
+	/* Locking for y rotation movement; Tilt or Pitch */
+	RotationY,
+
+	/* Locking for Z rotation movement; Pan or Yaw */
+	RotationZ
+};
+
+/**
+ * Stores relevant information for an axis for the virtual camera pawn
+ */
+USTRUCT()
+struct FVirtualCameraAxisSettings
+{
+	GENERATED_BODY()
+
+	/** If an axis is locked, then that axis's movement will be disabled. When unlocking an axis the movement is updated as if the axis had never been locked. */
+	UPROPERTY()
+	bool bIsLocked;
+
+	/** If an axis is frozen, then that frozen axis's movement will be disabled. When frozen axis is unlocked the movement is not updated, and all of the axis movement is applied as if it was in the location when the freeze was initiated. */
+	UPROPERTY()
+	bool bIsFrozen;
+
+	/** The amount of stabilization that can be applied to an axis */
+	UPROPERTY()
+	float StabilizationScale;
+
+	/** The scale that a user's movement should be adjusted by in game */
+	UPROPERTY()
+	float MovementScale;
+
+	/** The lock offset associated with an axis */
+	float LockRotationOffset;
+	FVector LockLocationOffset;
+
+	/** The freeze offset associated with an axis */
+	float FreezeRotationOffset;
+	FVector FreezeLocationOffset;
+
+	FVirtualCameraAxisSettings()
+	{
+		bIsLocked = false;
+		bIsFrozen = false;
+		StabilizationScale = .5f;
+		MovementScale = 1.f;
+		LockRotationOffset = 0.f;
+		LockLocationOffset = FVector::ZeroVector;
+		FreezeRotationOffset = 0.f;
+		FreezeLocationOffset = FVector::ZeroVector;
+	}
+
+	/**
+	 * Add an offset based on the freeze and lock states
+	 * @param InOffset - The offset to be added
+	 */
+	void AddRotationOffset(const float InOffset)
+	{
+		if (bIsFrozen)
+		{
+			FreezeRotationOffset += InOffset;
+		}
+		else if (bIsLocked)
+		{
+			LockRotationOffset += InOffset;
+		}
+	}
+
+	/**
+	 * Add an offset based on the freeze and lock states
+	 * @param InOffset - The offset to be added
+	 */
+	void AddLocationOffset(const FVector InOffset)
+	{
+		if (bIsFrozen)
+		{
+			FreezeLocationOffset += InOffset;
+		}
+		else if (bIsLocked)
+		{
+			LockLocationOffset += InOffset;
+		}
+	}
+
+	/**
+	 * Get the current total offset applied to this axis
+	 * @return - The total offset this axis has applied
+	 */
+	float GetRotationOffset() const { return FreezeRotationOffset + LockRotationOffset; }
+	FVector GetLocationOffset() const { return FreezeLocationOffset + LockLocationOffset; }
+
+	/** 
+	 * Helper method to check if an axis is prevented from moving.
+	 */
+	bool IsAxisImmobilized() const
+	{
+		return bIsLocked || bIsFrozen;
+	}
+
+	/**
+	 * Set the current lock state
+	 * @param bNewIsLocked - The desired lock state
+	 */
+	void SetIsLocked(const bool bNewIsLocked)
+	{
+		// If unlocking, clear locking offset
+		if (!bNewIsLocked)
+		{
+			LockRotationOffset = 0.f;
+			LockLocationOffset = FVector::ZeroVector;
+		}
+
+		bIsLocked = bNewIsLocked;
+	}
+
+	/**
+	 * Toggle the current lock state
+	 * @return - The current lock state after toggle
+	 */ 
+	bool ToggleLock()
+	{
+		SetIsLocked(!bIsLocked);
+		return bIsLocked;
+	}
+};
+
+/**
+ * Stores specific camera settings to be retrieved at a later time
+ */
+USTRUCT(BlueprintType)
+struct FVirtualCameraSettings
+{
+	GENERATED_BODY()
+
+	/* The focal length of the camera */
+	UPROPERTY(EditAnywhere, BlueprintReadWrite, Category = Basic)
+	float FocalLength;
+
+	/* The aperture of the camera */
+	UPROPERTY(EditAnywhere, BlueprintReadWrite, Category = Basic)
+	float Aperture;
+
+	UPROPERTY(EditAnywhere, Category = Basic)
+	FString FilmbackName;
+
+	/* The filmback sensor width of the camera */
+	UPROPERTY(EditAnywhere, Category = Basic)
+	float FilmbackWidth;
+	
+	/* The filmback sensor height of the camera */
+	UPROPERTY(EditAnywhere, Category = Basic)
+	float FilmbackHeight;
+
+	UPROPERTY(EditAnywhere, Category = Basic)
+	float MatteOpacity;
+	
+	/* The Axis settings for locking, stabilization, and scaling */
+	UPROPERTY(EditAnywhere, Category = Basic)
+	TMap<EVirtualCameraAxis, FVirtualCameraAxisSettings> AxisSettings;
+
+	UPROPERTY(EditAnywhere, Category = Basic)
+	bool bAllowFocusVisualization;
+
+	UPROPERTY(EditAnywhere, Category = Basic)
+	FColor DebugFocusPlaneColor;
+
+	UPROPERTY(VisibleAnywhere, Category = Basic)
+	EUnit DesiredDistanceUnits;
+
+	FVirtualCameraSettings()
+	{
+		FocalLength = 35.f;
+		Aperture = 2.8f;
+		FilmbackName = "";
+		FilmbackWidth = 24.889f;
+		FilmbackHeight = 14.f;
+		bAllowFocusVisualization = true;
+		DebugFocusPlaneColor = FColor(EForceInit::ForceInit);
+		MatteOpacity = .7f;
+		DesiredDistanceUnits = EUnit::Meters;
+
+		AxisSettings.Add(EVirtualCameraAxis::LocationX);
+		AxisSettings.Add(EVirtualCameraAxis::LocationY);
+		AxisSettings.Add(EVirtualCameraAxis::LocationZ);
+		AxisSettings.Add(EVirtualCameraAxis::RotationX);
+		AxisSettings.Add(EVirtualCameraAxis::RotationY);
+		AxisSettings.Add(EVirtualCameraAxis::RotationZ);
+	}
+};
+
+/**
+ * Keeps a saved location and the name of that location
+ */
+USTRUCT(BlueprintType)
+struct FVirtualCameraWaypoint
+{
+	GENERATED_BODY()
+
+	/* An associated name for the waypoint provided by the user */
+	UPROPERTY(BlueprintReadWrite, Category = Basic)
+	FString Name;
+
+	/* The location of the waypoint in world space */
+	UPROPERTY(BlueprintReadWrite, Category = Basic)
+	FTransform WaypointTransform;
+
+	UPROPERTY(BlueprintReadWrite, Category = Basic)
+	FDateTime DateCreated;
+
+	UPROPERTY(BlueprintReadWrite, Category = Basic)
+	bool bIsHomeWaypoint;
+
+	UPROPERTY(BlueprintReadWrite, Category = Basic)
+	bool bIsFavorited;
+
+	static int32 NextIndex;
+
+	FVirtualCameraWaypoint()
+	{
+		DateCreated = FDateTime(0);
+		bIsHomeWaypoint = false;
+		bIsFavorited = false;
+		Name = "";
+	};
+};
+
+/**
+* Keeps track of all the data associated with a screenshot that was taken by the user
+*/
+USTRUCT(BlueprintType)
+struct FVirtualCameraScreenshot
+{
+	GENERATED_BODY()
+
+	/* The name and location of the screenshot */
+	UPROPERTY(BlueprintReadWrite, Category = Basic)
+	FVirtualCameraWaypoint Waypoint;
+
+	/* The associated camera data from when the screenshot was taken */
+	UPROPERTY(BlueprintReadWrite, Category = Basic)
+	FVirtualCameraSettings CameraSettings;
+
+	static int32 NextIndex;
+};
+
+/**
+ * Keeps track of all data associated with settings presets.
+ */
+USTRUCT(BlueprintType)
+struct FVirtualCameraSettingsPreset
+{
+	GENERATED_BODY()
+
+	/* Checks which settings are saved for the preset */
+	UPROPERTY(EditAnywhere, BlueprintReadWrite, Category = Basic)
+	bool bIsCameraSettingsSaved;
+
+	UPROPERTY(EditAnywhere, BlueprintReadWrite, Category = Basic)
+	bool bIsStabilizationSettingsSaved;
+
+	UPROPERTY(EditAnywhere, BlueprintReadWrite, Category = Basic)
+	bool bIsAxisLockingSettingsSaved;
+
+	UPROPERTY(EditAnywhere, BlueprintReadWrite, Category = Basic)
+	bool bIsMotionScaleSettingsSaved;
+
+	/* Checks if saettings is set as favorite */
+	UPROPERTY(EditAnywhere, BlueprintReadWrite, Category = Basic)
+	bool bIsFavorited;
+
+	UPROPERTY(EditAnywhere, Category = Basic)
+	FVirtualCameraSettings CameraSettings;
+
+	UPROPERTY(VisibleAnywhere, BlueprintReadWrite, Category = Basic)
+	FDateTime DateCreated;
+
+	static int32 NextIndex;
+
+	FVirtualCameraSettingsPreset()
+	{
+		bIsCameraSettingsSaved = false;
+		bIsStabilizationSettingsSaved = false;
+		bIsAxisLockingSettingsSaved = false;
+		bIsMotionScaleSettingsSaved = false;
+		bIsFavorited = false;
+		DateCreated = FDateTime(0);
+	};
+};
+
+UCLASS()
+class VIRTUALCAMERA_API UVirtualCameraSaveGame : public USaveGame
+{
+	GENERATED_UCLASS_BODY()
+
+public:
+	UPROPERTY(VisibleAnywhere, Category = Basic)
+	TMap<FString, FVirtualCameraWaypoint> Waypoints;
+
+	UPROPERTY(VisibleAnywhere, Category = Basic)
+	TMap<FString, FVirtualCameraScreenshot> Screenshots;
+
+	UPROPERTY(VisibleAnywhere, Category = Basic)
+	FString HomeWaypointName;
+
+	UPROPERTY(VisibleAnywhere, Category = Basic)
+	FVirtualCameraSettings CameraSettings;
+
+	UPROPERTY(VisibleAnywhere, Category = Basic)
+	TMap<FString, FVirtualCameraSettingsPreset> SettingsPresets;
+
+	UPROPERTY(VisibleAnywhere, Category = Basic)
+	FString SaveSlotName;
+
+	UPROPERTY(VisibleAnywhere, Category = Basic)
+	uint32 UserIndex;
+
+	UPROPERTY(VisibleAnywhere, Category = Basic)
+	int32 WaypointIndex;
+
+	UPROPERTY(VisibleAnywhere, Category = Basic)
+	int32 ScreenshotIndex;
+
+	UPROPERTY(VisibleAnywhere, Category = Basic)
+	int32 PresetIndex;
+};
diff --git a/Engine/Plugins/Experimental/VirtualCamera/Source/VirtualCamera/Public/VirtualCameraSubsystem.h b/Engine/Plugins/Experimental/VirtualCamera/Source/VirtualCamera/Public/VirtualCameraSubsystem.h
new file mode 100644
--- /dev/null
+++ b/Engine/Plugins/Experimental/VirtualCamera/Source/VirtualCamera/Public/VirtualCameraSubsystem.h
@@ -0,0 +1,71 @@
+// Copyright Epic Games, Inc. All Rights Reserved.
+
+#pragma once
+
+#include "Delegates/DelegateCombinations.h"
+#include "IVirtualCameraController.h"
+#include "Subsystems/EngineSubsystem.h"
+#include "VirtualCameraSubsystem.generated.h"
+
+DECLARE_DYNAMIC_MULTICAST_DELEGATE(FOnStreamStarted);
+DECLARE_DYNAMIC_MULTICAST_DELEGATE(FOnStreamStopped);
+DECLARE_DYNAMIC_MULTICAST_DELEGATE_OneParam(FOnSelectedAnyActor, AActor*, SelectedActor);
+DECLARE_DYNAMIC_MULTICAST_DELEGATE_OneParam(FOnSelectedViewportActor, AActor*, SelectedActor);
+
+UCLASS(BlueprintType, Category = "VirtualCamera", DisplayName = "VirtualCameraSubsystem")
+class VIRTUALCAMERA_API UVirtualCameraSubsystem : public UEngineSubsystem
+{
+	GENERATED_BODY()
+	
+public:
+
+	UVirtualCameraSubsystem();
+
+public:
+	UFUNCTION(BlueprintCallable, Category = "VirtualCamera | Streaming")
+	bool StartStreaming();
+
+	UFUNCTION(BlueprintCallable, Category = "VirtualCamera | Streaming")
+	bool StopStreaming();
+
+	UFUNCTION(BlueprintCallable, Category = "VirtualCamera | Streaming")
+	bool IsStreaming() const;
+
+#if WITH_EDITOR
+	UFUNCTION()
+	void HandleSelectionChangedEvent(UObject* SelectedObject);
+
+	UFUNCTION()
+	void HandleSelectObjectEvent(UObject* SelectedObject);
+#endif
+
+	UFUNCTION(BlueprintCallable, Category = "VirtualCamera")
+	TScriptInterface<IVirtualCameraController> GetVirtualCameraController() const;
+	
+	UFUNCTION(BlueprintCallable, Category = "VirtualCamera")
+	void SetVirtualCameraController(TScriptInterface<IVirtualCameraController> VirtualCamera);
+
+	UPROPERTY(BlueprintReadOnly, Category = "VirtualCamera")
+	TObjectPtr<ULevelSequencePlaybackController> SequencePlaybackController;
+
+	UPROPERTY(BlueprintAssignable, Category = "VirtualCamera | Streaming")
+	FOnStreamStarted OnStreamStartedDelegate;
+
+	UPROPERTY(BlueprintAssignable, Category = "VirtualCamera | Streaming")
+	FOnStreamStopped OnStreamStoppedDelegate;
+	
+	UPROPERTY(BlueprintAssignable, Category = "VirtualCamera | Streaming")
+	FOnSelectedAnyActor OnSelectedAnyActorDelegate; 
+
+	UPROPERTY(BlueprintAssignable, Category = "VirtualCamera | Streaming")
+	FOnSelectedViewportActor OnSelectedActorInViewportDelegate;
+
+
+
+private:
+
+	UPROPERTY(Transient) 
+	TScriptInterface<IVirtualCameraController> ActiveCameraController;
+
+	bool bIsStreaming;
+};
diff --git a/Engine/Plugins/Experimental/VirtualCamera/Source/VirtualCamera/Public/VirtualCameraUserSettings.h b/Engine/Plugins/Experimental/VirtualCamera/Source/VirtualCamera/Public/VirtualCameraUserSettings.h
new file mode 100644
--- /dev/null
+++ b/Engine/Plugins/Experimental/VirtualCamera/Source/VirtualCamera/Public/VirtualCameraUserSettings.h
@@ -0,0 +1,142 @@
+// Copyright Epic Games, Inc. All Rights Reserved.
+
+#pragma once
+
+#include "CoreMinimal.h"
+#include "UObject/Object.h"
+#include "GameFramework/Actor.h"
+#include "GameFramework/PlayerInput.h"
+
+#include "VirtualCameraUserSettings.generated.h"
+
+/**
+ * Virtual Camera User Settings
+ */
+UCLASS(config=VirtualCamera, BlueprintType)
+class VIRTUALCAMERA_API UVirtualCameraUserSettings : public UObject
+{
+	GENERATED_BODY()
+
+public:
+
+	UVirtualCameraUserSettings();
+	
+	/** Controls interpolation speed when smoothing when changing focus distance. This is used to set the value of FocusSmoothingInterpSpeed in the Virtual camera CineCamera component */
+	UPROPERTY(EditAnywhere, config, Category = "VirtualCamera", meta = (ClampMin = "1.0", ClampMax = "50.0", DisplayName = "Focus Interpolation Speed"))
+	float FocusInterpSpeed = 8.0f;
+
+	/** Controls how fast the camera moves when using joysticks */
+	UPROPERTY(EditAnywhere, config, Category = "VirtualCamera", meta = (DisplayName = "Joysticks Speed"))
+	float JoysticksSpeed = 5.0f;
+
+	/** Sets the maximum possible joystick speed */
+	UPROPERTY(EditAnywhere, config, Category = "VirtualCamera", meta = (DisplayName = "Max Joysticks Speed"))
+	float MaxJoysticksSpeed = 10.0f;
+
+	/** Whether the map is displayed using grayscale or full color */
+	UPROPERTY(EditAnywhere, config, Category = "VirtualCamera", meta = (DisplayName = "Display Map In Grayscale"))
+	bool bIsMapGrayscale = true;
+	
+	/** Whether to change camera lens and fstop when teleporting to a screenshot to those with which the screenshot was taken */
+	UPROPERTY(EditAnywhere, config, Category = "VirtualCamera", meta = (DisplayName = "Override Camera Settings On Teleporting To Screenshot"))
+	bool bOverrideCameraSettingsOnTeleportToScreenshot = true;
+
+	/** Stores the filmback preset name selected by the user */
+	UPROPERTY(EditAnywhere, config, Category = "VirtualCamera", meta = (DisplayName = "Virtual Camera Filmback"))
+	FString VirtualCameraFilmback;
+
+	/** Whether to display film leader when recording a take */
+	UPROPERTY(EditAnywhere, config, Category = "VirtualCamera", meta = (DisplayName = "Display Film Leader"))
+	bool bDisplayFilmLeader = true;
+
+	/** Whether to teleport to the home bookmark when VCam starts */
+	UPROPERTY(EditAnywhere, config, Category = "VirtualCamera", meta = (DisplayName = "Teleport To Home On Start"))
+	bool bTeleportOnStart = true;
+
+	/** Default Vcam Class for Vcam Operator Panel */
+	UPROPERTY(EditAnywhere, BlueprintReadWrite, config, Category = "VirtualCamera|Presets")
+	TSoftClassPtr<AActor> DefaultVCamClass;
+
+	UPROPERTY(BlueprintReadWrite, config, Category = "VirtualCamera", meta = (DisplayName = "VirtualCamera Axis Mappings"))
+	TArray<struct FInputAxisKeyMapping> AxisMappings;
+
+	UPROPERTY(BlueprintReadWrite, config, Category = "VirtualCamera", meta = (DisplayName = "VirtualCamera Action Mappings"))
+	TArray<struct FInputActionKeyMapping> ActionMappings;
+
+	/** Get FocusInterpSpeed variable */
+	UFUNCTION(BlueprintPure, Category = "VirtualCamera")
+	float GetFocusInterpSpeed();
+
+	/** Set FocusInterpSpeed variable */
+	UFUNCTION(BlueprintCallable, Category = "VirtualCamera")
+	void SetFocusInterpSpeed(const float InFocusInterpSpeed);
+
+	/** Get JoysticksSpeed variable */
+	UFUNCTION(BlueprintPure, Category = "VirtualCamera")
+	float GetJoysticksSpeed();
+
+	/** Set JoysticksSpeed variable */
+	UFUNCTION(BlueprintCallable, Category = "VirtualCamera")
+	void SetJoysticksSpeed(const float InJoysticksSpeed);
+
+	/** Get MaxJoysticksSpeed variable */
+	UFUNCTION(BlueprintPure, Category = "VirtualCamera")
+	float GetMaxJoysticksSpeed();
+
+	/** Set MaxJoysticksSpeed variable */
+	UFUNCTION(BlueprintCallable, Category = "VirtualCamera")
+	void SetMaxJoysticksSpeed(const float InMaxJoysticksSpeed);
+
+	/** Get bIsMapGrayscale variable */
+	UFUNCTION(BlueprintPure, Category = "VirtualCamera")
+	bool IsMapGrayscle();
+
+	/** Set bIsMapGrayscale variable */
+	UFUNCTION(BlueprintCallable, Category = "VirtualCamera")
+	void SetIsMapGrayscle(const bool bInIsMapGrayscle);
+
+	/** Get bOverrideCameraSettingsOnTeleportToScreenshot variable */
+	UFUNCTION(BlueprintPure, Category = "VirtualCamera")
+	bool GetShouldOverrideCameraSettingsOnTeleport();
+
+	/** Set bOverrideCameraSettingsOnTeleportToScreenshot variable */
+	UFUNCTION(BlueprintCallable, Category = "VirtualCamera")
+	void SetShouldOverrideCameraSettingsOnTeleport(const bool bInOverrideCameraSettings);
+
+	/** Get VirtualCameraFilmback variable */
+	UFUNCTION(BlueprintPure, Category = "VirtualCamera")
+	FString GetSavedVirtualCameraFilmbackPresetName();
+
+	/** Set VirtualCameraFilmback variable */
+	UFUNCTION(BlueprintCallable, Category = "VirtualCamera")
+	void SetSavedVirtualCameraFilmbackPresetName(const FString& InFilmback);
+
+	/** Get bDisplayFilmLeader variable */
+	UFUNCTION(BlueprintPure, Category = "VirtualCamera")
+	bool GetShouldDisplayFilmLeader();
+
+	/** Set bDisplayFilmLeader variable */
+	UFUNCTION(BlueprintCallable, Category = "VirtualCamera")
+	void SetShouldDisplayFilmLeader(const bool bInDisplayFilmLeader);
+
+	/** Get bTeleportOnStart variable */
+	UFUNCTION(BlueprintPure, Category = "VirtualCamera")
+	bool GetTeleportOnStart();
+
+	/** Set bTeleportOnStart variable */
+	UFUNCTION(BlueprintCallable, Category = "VirtualCamera")
+	void SetTeleportOnStart(const bool bInTeleportOnStart);
+
+	/** Fills the Axis/Action mappings with assosiated gamepad bindings */
+	UFUNCTION(BlueprintCallable, Category = "VirtualCamera")
+	void InjectGamepadKeybinds(); 
+	
+	/** Retrieve all VirtualCamera action mappings by a certain name. */
+	UFUNCTION(BlueprintPure, Category = Settings)
+	void GetActionMappingsByName(const FName InActionName, TArray<FInputActionKeyMapping>& OutMappings) const;
+
+	/** Retrieve all VirtualCamera axis mappings by a certain name. */
+	UFUNCTION(BlueprintPure, Category = Settings)
+	void GetAxisMappingsByName(const FName InAxisName, TArray<FInputAxisKeyMapping>& OutMappings) const;
+
+};
\ No newline at end of file
diff --git a/Engine/Plugins/Experimental/VirtualCamera/Source/VirtualCamera/Public/VirtualCameraVersion.h b/Engine/Plugins/Experimental/VirtualCamera/Source/VirtualCamera/Public/VirtualCameraVersion.h
new file mode 100644
--- /dev/null
+++ b/Engine/Plugins/Experimental/VirtualCamera/Source/VirtualCamera/Public/VirtualCameraVersion.h
@@ -0,0 +1,5 @@
+// Copyright Epic Games, Inc. All Rights Reserved.
+
+#pragma once
+
+#define VIRTUAL_CAMERA_VERSION_STRING TEXT("1.0.0")
diff --git a/Engine/Plugins/Experimental/VirtualCamera/Source/VirtualCamera/VirtualCamera.Build.cs b/Engine/Plugins/Experimental/VirtualCamera/Source/VirtualCamera/VirtualCamera.Build.cs
new file mode 100644
--- /dev/null
+++ b/Engine/Plugins/Experimental/VirtualCamera/Source/VirtualCamera/VirtualCamera.Build.cs
@@ -0,0 +1,73 @@
+// Copyright Epic Games, Inc. All Rights Reserved.
+
+using UnrealBuildTool;
+
+public class VirtualCamera : ModuleRules
+{
+	public VirtualCamera(ReadOnlyTargetRules Target) : base(Target)
+	{
+		DefaultBuildSettings = BuildSettingsVersion.V2;
+
+		PublicDependencyModuleNames.AddRange(
+			new string[]
+			{
+				"AugmentedReality",
+				"CinematicCamera",
+				"Core",
+				"CoreUObject",
+				"Engine",
+				"InputCore",
+				"LevelSequence",
+				"LiveLinkInterface",
+				"MovieScene",
+				"RemoteSession",
+				"TimeManagement",
+				"UMG",
+				"VCamCore",
+				"VPUtilities",
+				"AssetRegistry",
+			}
+		);
+
+		PrivateDependencyModuleNames.AddRange(
+			new string[]
+			{
+				"MediaIOCore",
+				"Slate",
+				"AdvancedWidgets"
+			}
+		);
+
+		if (Target.Type == TargetType.Editor || Target.Type == TargetType.Program)
+		{
+			PrivateDefinitions.Add("VIRTUALCAMERA_WITH_CONCERT=1");
+			PrivateDependencyModuleNames.AddRange(
+				new string[]
+				{
+					"Concert",
+					"ConcertSyncClient",
+					"MultiUserClient",
+				}
+			);
+		}
+		else
+		{
+			PrivateDefinitions.Add("VIRTUALCAMERA_WITH_CONCERT=0");
+		}
+
+
+		if (Target.bBuildEditor == true)
+		{
+			PublicDependencyModuleNames.Add("LevelSequenceEditor");
+			PublicDependencyModuleNames.Add("Sequencer");
+			PublicDependencyModuleNames.Add("SlateCore");
+			PublicDependencyModuleNames.Add("TakeRecorder");
+			PrivateDependencyModuleNames.Add("LevelEditor");
+			PrivateDependencyModuleNames.Add("EditorFramework");
+			PrivateDependencyModuleNames.Add("UnrealEd");
+			PrivateDependencyModuleNames.Add("EditorScriptingUtilities");
+			PrivateDependencyModuleNames.Add("VPUtilitiesEditor");
+			PrivateDependencyModuleNames.Add("TakesCore");
+		}
+	}
+}
diff --git a/Engine/Plugins/Experimental/VirtualCamera/Source/VirtualCameraEditor/Private/VirtualCameraEditorModule.cpp b/Engine/Plugins/Experimental/VirtualCamera/Source/VirtualCameraEditor/Private/VirtualCameraEditorModule.cpp
new file mode 100644
--- /dev/null
+++ b/Engine/Plugins/Experimental/VirtualCamera/Source/VirtualCameraEditor/Private/VirtualCameraEditorModule.cpp
@@ -0,0 +1,52 @@
+// Copyright Epic Games, Inc. All Rights Reserved.
+
+#include "VirtualCameraEditorModule.h"
+
+#include "ISettingsModule.h"
+#include "ISettingsSection.h"
+#include "Modules/ModuleManager.h"
+#include "VirtualCameraUserSettings.h"
+
+#define LOCTEXT_NAMESPACE "FVirtualCameraEditorModule"
+
+DEFINE_LOG_CATEGORY_STATIC(LogVirtualCameraEditor, Log, All);
+
+class FVirtualCameraEditorModule : public IModuleInterface
+{
+public:
+	virtual void StartupModule() override
+	{
+		RegisterSettings();
+	}
+
+	virtual void ShutdownModule() override
+	{
+		UnregisterSettings();
+	}
+
+	void RegisterSettings()
+	{
+		ISettingsModule* SettingsModule = FModuleManager::GetModulePtr<ISettingsModule>("Settings");
+		if (SettingsModule != nullptr)
+		{
+			ISettingsSectionPtr SettingsSection = SettingsModule->RegisterSettings("Project", "Plugins", "VirtualCamera",
+				LOCTEXT("VirtualCameraUserSettingsName", "Virtual Camera"),
+				LOCTEXT("VirtualCameraUserSettingsDescription", "Configure the Virtual Camera settings."),
+				GetMutableDefault<UVirtualCameraUserSettings>());
+		}
+	}
+
+	void UnregisterSettings()
+	{
+		ISettingsModule* SettingsModule = FModuleManager::GetModulePtr<ISettingsModule>("Settings");
+		if (SettingsModule != nullptr)
+		{
+			SettingsModule->UnregisterSettings("Project", "Plugins", "VirtualCamera");
+		}
+	}
+};
+#undef LOCTEXT_NAMESPACE
+
+IMPLEMENT_MODULE(FVirtualCameraEditorModule, VirtualCameraEditor)
+
+
diff --git a/Engine/Plugins/Experimental/VirtualCamera/VirtualCamera.uplugin b/Engine/Plugins/Experimental/VirtualCamera/VirtualCamera.uplugin
--- a/Engine/Plugins/Experimental/VirtualCamera/VirtualCamera.uplugin
+++ b/Engine/Plugins/Experimental/VirtualCamera/VirtualCamera.uplugin
@@ -32,6 +32,20 @@
 			"Name": "VCamExtensionsEditor",
 			"Type": "Editor",
 			"LoadingPhase": "Default"
+		},
+		{
+			"Name": "VirtualCamera",
+			"Type": "Runtime",
+			"LoadingPhase": "Default",
+			"PlatformAllowList": [
+				"Win64",
+				"Linux"
+			]
+		},
+		{
+			"Name": "VirtualCameraEditor",
+			"Type": "Editor",
+			"LoadingPhase": "PostEngineInit"
 		}
 	],
 	"Plugins": [
@@ -46,6 +60,6 @@
 		{
 			"Name" : "VCamPixelStreamingPreset",
 			"Enabled": true
-		},
+		}
 	]
 }
