Enable VirtualCamera for all platforms, specifically for Mac

diff --git a/Engine/Plugins/Experimental/VirtualCamera/VirtualCamera.uplugin b/Engine/Plugins/Experimental/VirtualCamera/VirtualCamera.uplugin
--- a/Engine/Plugins/Experimental/VirtualCamera/VirtualCamera.uplugin
+++ b/Engine/Plugins/Experimental/VirtualCamera/VirtualCamera.uplugin
@@ -14,14 +14,6 @@
 	"IsBetaVersion": true,
 	"Installed": false,
 	"EnabledByDefault": false,
- 	"SupportedTargetPlatforms": [
- 		"Win64",
- 		"Linux"
- 	],
- 	"PlatformAllowList": [
- 		"Win64",
- 		"Linux"
- 	],
 	"Modules": [
 		{
 			"Name": "VCamExtensions",
