Restore file incorrectly deleted by perforce's completely broken move resolution. Pray it stays restored this time.
#rnx

diff --git a/Engine/Plugins/Experimental/VirtualCamera/Source/VirtualCamera/Public/LevelSequence/VirtualCameraClipsMetaData.h b/Engine/Plugins/Experimental/VirtualCamera/Source/VirtualCamera/Public/LevelSequence/VirtualCameraClipsMetaData.h
new file mode 100644
--- /dev/null
+++ b/Engine/Plugins/Experimental/VirtualCamera/Source/VirtualCamera/Public/LevelSequence/VirtualCameraClipsMetaData.h
@@ -0,0 +1,274 @@
+// Copyright Epic Games, Inc. All Rights Reserved.
+
+#pragma once
+
+
+#include "UObject/Object.h"
+#include "IMovieSceneMetaData.h"
+#include "Misc/FrameRate.h"
+#include "VirtualCameraClipsMetaData.generated.h"
+
+
+/**
+ * Clips meta-data that is stored on ULevelSequence assets that are recorded through the virtual camera. 
+ * Meta-data is retrieved through ULevelSequence::FindMetaData<UVirtualCameraClipsMetaData>()
+ */
+
+UCLASS(BlueprintType)
+class UVirtualCameraClipsMetaData : public UObject, public IMovieSceneMetaDataInterface
+{
+public: 
+	GENERATED_BODY()
+	UVirtualCameraClipsMetaData(const FObjectInitializer& ObjInit);
+
+public: 
+
+	/** The asset registry tag that contains the focal length for this meta-data */
+	static const FName AssetRegistryTag_FocalLength;
+
+	/** The asset registry tag that contains if the selected state for this meta-data */
+	static const FName AssetRegistryTag_bIsSelected;
+
+	/** The asset registry tag that contains the recorded level name for this meta-data */
+	static const FName AssetRegistryTag_RecordedLevelName;
+
+	/** The asset registry tag that contains the FrameCountStart in for this meta-data */
+	static const FName AssetRegistryTag_FrameCountStart;
+
+	/** The asset registry tag that contains the FrameCountEnd out for this meta-data */
+	static const FName AssetRegistryTag_FrameCountEnd;
+
+	/** The asset registry tag that contains the LengthInFrames out for this meta-data */
+	static const FName AssetRegistryTag_LengthInFrames; 
+
+	/** The asset registry tag that contains the FrameCountEnd out for this meta-data */
+	static const FName AssetRegistryTag_DisplayRate;
+
+	/** The asset registry tag that contains whether the clip was recorded with a CineCamera for this meta-data */
+	static const FName AssetRegistryTag_bIsACineCameraRecording;
+	
+	/** The asset registry tag that contains whether this take is good or not */
+	static const FName AssetRegistryTag_bIsNoGood;
+
+	/** The asset registry tag that contains whether this was flagged by a user */
+	static const FName AssetRegistryTag_bIsFlagged;
+
+	/** The asset registry tag that contains its favorite status */
+	static const FName AssetRegistryTag_FavoriteLevel;
+
+	/** The asset registry tag that contains whether it was created from a VCam */
+	static const FName AssetRegistryTag_bIsCreatedFromVCam;
+	
+public:
+
+	/** The asset registry tag that contains the focal length for this meta-data */
+	UFUNCTION(BlueprintPure, Category = "VirtualCamera|Clips")
+	static FName GetClipsMetaDataTag_FocalLength() { return AssetRegistryTag_FocalLength; }
+	
+	/** The asset registry tag that contains if the selected state for this meta-data */
+	UFUNCTION(BlueprintPure, Category = "VirtualCamera|Clips")
+	static FName GetClipsMetaDataTag_IsSelected() { return AssetRegistryTag_bIsSelected; }
+	
+	/** The asset registry tag that contains the recorded level name for this meta-data */
+	UFUNCTION(BlueprintPure, Category = "VirtualCamera|Clips")
+	static FName GetClipsMetaDataTag_RecordedLevel() { return AssetRegistryTag_RecordedLevelName; }
+	
+	/** The asset registry tag that contains the FrameCountStart in for this meta-data */
+	UFUNCTION(BlueprintPure, Category = "VirtualCamera|Clips")
+	static FName GetClipsMetaDataTag_FrameCountStart() { return AssetRegistryTag_FrameCountStart; }
+	
+	/** The asset registry tag that contains the FrameCountEnd out for this meta-data */
+	UFUNCTION(BlueprintPure, Category = "VirtualCamera|Clips")
+	static FName GetClipsMetaDataTag_FrameCountEnd() { return AssetRegistryTag_FrameCountEnd; }
+	
+	/** The asset registry tag that contains the LengthInFrames out for this meta-data */
+	UFUNCTION(BlueprintPure, Category = "VirtualCamera|Clips")
+	static FName GetClipsMetaDataTag_LengthInFrames() { return AssetRegistryTag_LengthInFrames; }
+	
+	/** The asset registry tag that contains the FrameCountEnd out for this meta-data */
+	UFUNCTION(BlueprintPure, Category = "VirtualCamera|Clips")
+	static FName GetClipsMetaDataTag_DisplayRate() { return AssetRegistryTag_DisplayRate; }
+	
+	/** The asset registry tag that contains whether the clip was recorded with a CineCamera for this meta-data */
+	UFUNCTION(BlueprintPure, Category = "Virtual Camera|Clips")
+	static FName GetClipsMetaDataTag_IsCineACineCameraRecording() { return AssetRegistryTag_bIsACineCameraRecording; }
+	
+	/** The asset registry tag that contains whether this take is good or not */
+	UFUNCTION(BlueprintPure, Category = "VirtualCamera|Clips")
+	static FName GetClipsMetaDataTag_IsNoGood() { return AssetRegistryTag_bIsNoGood; }
+	
+	/** The asset registry tag that contains whether this was flagged by a user */
+	UFUNCTION(BlueprintPure, Category = "VirtualCamera|Clips")
+	static FName GetClipsMetaDataTag_IsFlagged() { return AssetRegistryTag_bIsFlagged; }
+	
+	/** The asset registry tag that contains its favorite status */
+	UFUNCTION(BlueprintPure, Category = "VirtualCamera|Clips")
+	static FName GetClipsMetaDataTag_FavoriteLevel() { return AssetRegistryTag_FavoriteLevel; }
+	
+	/** The asset registry tag that contains whether it was created from a VCam */
+	UFUNCTION(BlueprintPure, Category = "VirtualCamera|Clips")
+	static FName GetClipsMetaDataTag_IsCreatedFromVCam() { return AssetRegistryTag_bIsCreatedFromVCam; }
+
+	/** Gets all asset registry tags */
+	UFUNCTION(BlueprintPure, Category = "VirtualCamera|Clips")
+	static TSet<FName> GetAllClipsMetaDataTags();
+
+public:
+
+	/**
+	* Extend the default ULevelSequence asset registry tags
+	*/
+	virtual void ExtendAssetRegistryTags(TArray<FAssetRegistryTag>& OutTags) const override;
+#if WITH_EDITOR
+	virtual void ExtendAssetRegistryTagMetaData(TMap<FName, FAssetRegistryTagMetadata>& OutMetadata) const override;
+#endif
+
+public:
+	/**
+	 * @return The focal length for this clip
+	 */
+	UFUNCTION(BlueprintCallable, Category = "Clips")
+	float GetFocalLength() const;
+
+	/**
+	* @return Whether or not the clip is selected.  
+	*/
+	UFUNCTION(BlueprintCallable, Category = "Clips")
+	bool GetSelected() const;
+
+	/**
+	* @return The name of the clip's recorded level. 
+	*/
+	UFUNCTION(BlueprintCallable, Category = "Clips")
+	FString GetRecordedLevelName() const;
+
+	/**
+	* @return The initial frame of the clip
+	*/
+	UFUNCTION(BlueprintCallable, Category = "Clips")
+	int GetFrameCountStart() const;
+
+	/**
+	* @return The final frame of the clip
+	*/
+	UFUNCTION(BlueprintCallable, Category = "Clips")
+	int GetFrameCountEnd() const;
+
+	/**
+	* @return The length in frames of the clip. 
+	*/
+	UFUNCTION(BlueprintCallable, Category = "Clips")
+	int GetLengthInFrames();
+
+	/**
+	* @return The display rate of the clip. 
+	*/
+	UFUNCTION(BlueprintCallable, Category = "Clips")
+	FFrameRate GetDisplayRate();
+
+	/**
+	* @return Whether the clip was recorded by a CineCameraActor
+	*/
+	UFUNCTION(BlueprintCallable, Category = "Clips")
+	bool GetIsACineCameraRecording() const;
+
+public:
+	/**
+	* Set the focal length associated with this clip. 
+	* @note: Used for tracking. Does not update the StreamedCameraComponent. 
+	*/
+	UFUNCTION(BlueprintCallable, Category = "Clips")
+	void SetFocalLength(float InFocalLength);
+	
+	/**
+	* Set if this clip is 'selected'
+	*/
+	UFUNCTION(BlueprintCallable, Category = "Clips")
+	void SetSelected(bool bInSelected);
+
+	/**
+	* Set the name of the level that the clip was recorded in. 
+	*/
+	UFUNCTION(BlueprintCallable, Category = "Clips")
+	void SetRecordedLevelName(FString InLevelName);
+
+	/**
+	* Set the initial frame of the clip used for calculating duration.
+	*/
+	UFUNCTION(BlueprintCallable, Category = "Clips")
+	void SetFrameCountStart(int InFrame); 
+
+	/**
+	* Set the final frame of the clip used for calculating duration.
+	*/
+	UFUNCTION(BlueprintCallable, Category = "Clips")
+	void SetFrameCountEnd(int InFrame);
+
+	/**
+	 * Set the length in frames of the clip used for AssetData calculations. 
+	 */
+	UFUNCTION(BlueprintCallable, Category = "Clips")
+	void SetLengthInFrames(int InLength);
+
+	/**
+	 * Set the DisplayRate of the clip used for AssetData calculations.
+	 */
+	UFUNCTION(BlueprintCallable, Category = "Clips")
+	void SetDisplayRate(FFrameRate InDisplayRate);
+
+	/**
+	 * Set if the clip was recorded by a CineCameraActor
+	 */
+	UFUNCTION(BlueprintCallable, Category = "Clips")
+	void SetIsACineCameraRecording(bool bInIsACineCameraRecording);
+
+private:
+	
+	/** The focal length of the streamed camera used to record the take */
+	UPROPERTY(EditAnywhere, Category = "Clips", BlueprintReadWrite, meta = (AllowPrivateAccess = "true"))
+	float FocalLength;
+
+	/** Whether or not the take was marked as 'selected' */
+	UPROPERTY(EditAnywhere, Category = "Clips", BlueprintReadWrite, meta = (AllowPrivateAccess = "true"))
+	bool bIsSelected; 
+
+	/** The name of the level that the clip was recorded in */
+	UPROPERTY(EditAnywhere, Category = "Clips", BlueprintReadWrite, meta = (AllowPrivateAccess = "true"))
+	FString RecordedLevelName; 
+
+	/** The initial frame of the clip used for calculating duration. */
+	UPROPERTY(EditAnywhere, Category = "Clips", BlueprintReadWrite, meta = (AllowPrivateAccess = "true"))
+	int FrameCountStart;
+
+	/** The last frame of the clip used for calculating duration. */
+	UPROPERTY(EditAnywhere, Category = "Clips", BlueprintReadWrite, meta = (AllowPrivateAccess = "true"))
+	int FrameCountEnd; 
+
+	/** The level sequence length in frames calculated from VirtualCameraSubsystem used for AssetData calculations */
+	UPROPERTY(EditAnywhere, Category = "Clips", BlueprintReadWrite, meta = (AllowPrivateAccess = "true"))
+	int LengthInFrames; 
+
+	/** The display rate of the level sequence used for AssetData calculations. */
+	UPROPERTY(EditAnywhere, Category = "Clips", BlueprintReadWrite, meta = (AllowPrivateAccess = "true"))
+	FFrameRate DisplayRate;
+
+	/** If the LevelSequence was recorded with a CineCameraActor, rather than a VirtualCameraActor */
+	UPROPERTY(EditAnywhere, Category = "Clips", BlueprintReadWrite, meta = (AllowPrivateAccess = "true"))
+	bool bIsACineCameraRecording;
+	
+	/** Whether this take is marked as good */
+	UPROPERTY(EditAnywhere, Category = "Clips", BlueprintReadWrite, meta = (AllowPrivateAccess = "true"))
+	bool bIsNoGood = false;
+	
+	/** The asset registry tag that contains whether this was flagged by a user */
+	UPROPERTY(EditAnywhere, Category = "Clips", BlueprintReadWrite, meta = (AllowPrivateAccess = "true"))
+	bool bIsFlagged = false;
+
+	/** The asset registry tag that contains its favorite status */
+	UPROPERTY(EditAnywhere, Category = "Clips", BlueprintReadWrite, meta = (AllowPrivateAccess = "true"))
+	int32 FavoriteLevel = 0;
+
+	/** Whether the sequence was created from a VCam */
+	UPROPERTY(EditAnywhere, Category = "Clips", BlueprintReadWrite, meta = (AllowPrivateAccess = "true"))
+	bool bIsCreatedFromVCam = true;
+};
\ No newline at end of file
