Move private helper function to public interface and expose in blueprints.  This is for VCam widgets during hit testing.

diff --git a/Engine/Plugins/Experimental/VirtualCamera/Source/VirtualCamera/Public/FunctionLibraries/VCamBlueprintFunctionLibrary.h b/Engine/Plugins/Experimental/VirtualCamera/Source/VirtualCamera/Public/FunctionLibraries/VCamBlueprintFunctionLibrary.h
--- a/Engine/Plugins/Experimental/VirtualCamera/Source/VirtualCamera/Public/FunctionLibraries/VCamBlueprintFunctionLibrary.h
+++ b/Engine/Plugins/Experimental/VirtualCamera/Source/VirtualCamera/Public/FunctionLibraries/VCamBlueprintFunctionLibrary.h
@@ -191,12 +191,14 @@ public:
 	UFUNCTION(BlueprintCallable, Category = "VirtualCamera|Sequencer")
 	static void SetPlaybackSpeed(float Value=1.0);
 
-private:
-
+	/* Convert 2D screen position to World Space 3D position and direction. Returns false if unable to determine value. */
+	UFUNCTION(BlueprintCallable, Category = "VirtualCamera")
 	static bool DeprojectScreenToWorld(const FVector2D& InScreenPosition, FVector& OutWorldPosition, FVector& OutWorldDirection);
 
+private:
+
 #if WITH_EDITOR
 	/** Returns the current sequencer. */
 	static TWeakPtr<ISequencer> GetSequencer();
 #endif
-};
\ No newline at end of file
+};
