Expose functions for hit proxy tracing

diff --git a/Engine/Plugins/Experimental/VirtualCamera/Source/VirtualCamera/Private/FunctionLibraries/VCamBlueprintFunctionLibrary.cpp b/Engine/Plugins/Experimental/VirtualCamera/Source/VirtualCamera/Private/FunctionLibraries/VCamBlueprintFunctionLibrary.cpp
--- a/Engine/Plugins/Experimental/VirtualCamera/Source/VirtualCamera/Private/FunctionLibraries/VCamBlueprintFunctionLibrary.cpp
+++ b/Engine/Plugins/Experimental/VirtualCamera/Source/VirtualCamera/Private/FunctionLibraries/VCamBlueprintFunctionLibrary.cpp
@@ -8,12 +8,16 @@
 #include "AssetRegistry/AssetData.h"
 #include "CineCameraActor.h"
 #include "CineCameraComponent.h"
+#include "Components/PrimitiveComponent.h"
 #include "Components/SceneCaptureComponent2D.h"
+#include "EngineUtils.h"
 #include "Engine/GameInstance.h"
+#include "FunctionLibraries/TakeMetaDataTagsFunctionLibrary.h"
+#include "GameFramework/Actor.h"
 #include "GameFramework/PlayerController.h"
 #include "LevelSequence.h"
 #include "MovieScene.h"
-#include "FunctionLibraries/TakeMetaDataTagsFunctionLibrary.h"
+#include "Slate/SceneViewport.h"
 
 #if WITH_EDITOR
 #include "Editor.h"
@@ -572,6 +576,150 @@ bool UVCamBlueprintFunctionLibrary::DeprojectScreenToWorldByViewport(const FVect
 	return false;
 }
 
+namespace UE::VirtualCamera::Private
+{
+	enum class EBreakBehaviour : uint8
+	{
+		Break,
+		Continue
+	};
+
+	/**
+	 * @return SizeX, SizeY, and a SizeX x SizeY matrix where each element maps to a pixel.
+	 * @see FViewport::GetHitProxyMap.
+	 */
+	static TTuple<int32, int32, TArray<HHitProxy*>> GetProxyMap(FViewport& Viewport, const FVector2D& InScreenPosition, const uint32 HitProxySize)
+	{
+		// See FViewport::GetHitProxy.
+		// Compute a HitProxySize x HitProxySize test region with the center at (X,Y).
+		int32 MinX = InScreenPosition.X - HitProxySize;
+		int32 MinY = InScreenPosition.Y - HitProxySize;
+		int32 MaxX = InScreenPosition.X + HitProxySize;
+		int32 MaxY = InScreenPosition.Y + HitProxySize;
+		
+		// Clip the region to the viewport bounds.
+		const FIntPoint ViewportSize = Viewport.GetSizeXY();
+		MinX = FMath::Clamp(MinX, 0, ViewportSize.X - 1);
+		MinY = FMath::Clamp(MinY, 0, ViewportSize.Y - 1);
+		MaxX = FMath::Clamp(MaxX, 0, ViewportSize.X - 1);
+		MaxY = FMath::Clamp(MaxY, 0, ViewportSize.Y - 1);
+
+		int32 TestSizeX = MaxX - MinX + 1;
+		int32 TestSizeY = MaxY - MinY + 1;
+		TArray<HHitProxy*> Result;
+		
+		if (TestSizeX <= 0 || TestSizeY <= 0)
+		{
+			return { TestSizeX, TestSizeY, Result };
+		}
+		
+		const FIntRect QueryRect(MinX, MinY, MaxX + 1, MaxY + 1);
+		Viewport.GetHitProxyMap(QueryRect, Result);
+		return { TestSizeX, TestSizeY, Result };
+	}
+
+	/** Enumerates all hit proxies that correspond to a component. */
+	static void EnumerateHitProxies(
+		FViewport& Viewport,
+		const FVector2D& InScreenPosition,
+		const uint32 HitProxySize,
+		TFunctionRef<EBreakBehaviour(const UPrimitiveComponent& Component)> Callback
+		)
+	{
+		auto[TestSizeX, TestSizeY, ProxyMap] = GetProxyMap(Viewport, InScreenPosition, HitProxySize);
+		if (ProxyMap.IsEmpty())
+		{
+			return;
+		}
+		
+		const auto ProcessProxy = [&Callback](HHitProxy* Proxy) -> EBreakBehaviour
+		{
+			if (Proxy && Proxy->IsA(HActor::StaticGetType()))
+			{
+				HActor* ActorProxy = static_cast<HActor*>(Proxy);
+				const UPrimitiveComponent* Component = ActorProxy->PrimComponent;
+				if (Component)
+				{
+					return Callback(*Component);
+				}
+			}
+
+			return EBreakBehaviour::Continue;
+		};
+
+		// Process at the center first - if that it matches our requirements, it is our best option.
+		const int32 CenterIndex = TestSizeY/2 * TestSizeX + TestSizeX/2;
+		if (ProcessProxy(ProxyMap[CenterIndex]) == EBreakBehaviour::Break)
+		{
+			return;
+		}
+			
+		for (int32 TestY = 0;TestY < TestSizeY; TestY++)
+		{
+			for (int32 TestX = 0;TestX < TestSizeX; TestX++)
+			{
+				HHitProxy* HitProxy = ProxyMap[TestY * TestSizeX + TestX];
+				if (ProcessProxy(HitProxy) == EBreakBehaviour::Break)
+				{
+					return;
+				}
+			}
+		}
+	}
+
+	static bool PassesFilters(const UPrimitiveComponent& Component, const FVCamTraceHitProxyQueryParams& Params)
+	{
+		return !Params.IgnoredActors.Contains(Component.GetOwner());
+	}
+}
+
+bool UVCamBlueprintFunctionLibrary::MultiTraceHitProxyOnViewport(
+	const FVector2D& InScreenPosition,
+	EVCamTargetViewportID InTargetViewport,
+	FVCamTraceHitProxyQueryParams InQueryParams,
+	TArray<FVCamTraceHitProxyResult>& Result
+	)
+{
+	// This is WITH_EDITOR because EVCamTargetViewportID only makes sense in editor builds.
+	// Other than that, all the below functions actually exist in Runtime builds and would work.
+#if WITH_EDITOR
+	const TSharedPtr<SLevelViewport> ViewportWidget = UE::VCamCore::GetLevelViewport(InTargetViewport);
+	const TSharedPtr<FSceneViewport> Viewport = ViewportWidget ? ViewportWidget->GetSceneViewport() : nullptr;
+	if (!Viewport)
+	{
+		return false;
+	}
+
+	using namespace UE::VirtualCamera::Private;
+	TArray<FVCamTraceHitProxyResult> HitResults;
+	EnumerateHitProxies(*Viewport, InScreenPosition, FMath::Max(InQueryParams.HitProxySize, 0), [&InQueryParams, &HitResults](const UPrimitiveComponent& Component)
+	{
+		if (PassesFilters(Component, InQueryParams))
+		{
+			// Component is in a HActor proxy, which keeps a const reference. We cannot safely use const_cast without invoking undefined behavior.
+			// Hence we must resort to this runtime search hack.
+			TArray<UActorComponent*> Components = Component.GetOwner()->GetComponents().Array();
+			UActorComponent** UnconstComponent = Components.FindByPredicate([&Component](const UActorComponent* OwnedComponent){ return &Component == OwnedComponent; });
+			if (ensure(UnconstComponent))
+			{
+				UPrimitiveComponent* UnconstPrimitive = Cast<UPrimitiveComponent>(*UnconstComponent);
+				HitResults.AddUnique({ UnconstPrimitive->GetOwner(), UnconstPrimitive });
+			}
+			
+		}
+		return EBreakBehaviour::Continue;
+	});
+
+	if (!HitResults.IsEmpty())
+	{
+		Result = MoveTemp(HitResults);
+	}
+	return !Result.IsEmpty();
+#else
+	return false;
+#endif
+}
+
 TArray<UObject*> UVCamBlueprintFunctionLibrary::GetBoundObjects(FMovieSceneObjectBindingID CameraBindingID)
 {
 	TArray<UObject*> BoundObjectsArray;
diff --git a/Engine/Plugins/Experimental/VirtualCamera/Source/VirtualCamera/Public/FunctionLibraries/VCamBlueprintFunctionLibrary.h b/Engine/Plugins/Experimental/VirtualCamera/Source/VirtualCamera/Public/FunctionLibraries/VCamBlueprintFunctionLibrary.h
--- a/Engine/Plugins/Experimental/VirtualCamera/Source/VirtualCamera/Public/FunctionLibraries/VCamBlueprintFunctionLibrary.h
+++ b/Engine/Plugins/Experimental/VirtualCamera/Source/VirtualCamera/Public/FunctionLibraries/VCamBlueprintFunctionLibrary.h
@@ -10,9 +10,11 @@
 #include "MovieSceneObjectBindingID.h"
 #include "VCamBlueprintFunctionLibrary.generated.h"
 
+class AActor;
 class ACineCameraActor;
 class UCineCameraComponent;
 class ULevelSequence;
+class UPrimitiveComponent;
 class USceneCaptureComponent2D;
 class UVirtualCameraClipsMetaData;
 class UVirtualCameraUserSettings; 
@@ -23,6 +25,41 @@ class ISequencer;
 
 enum class EVCamTargetViewportID : uint8;
 
+USTRUCT(BlueprintType)
+struct VIRTUALCAMERA_API FVCamTraceHitProxyQueryParams
+{
+	GENERATED_BODY()
+
+	/** Determine the size of the query area around the center pixel. */
+	UPROPERTY(EditAnywhere, BlueprintreadWrite, Category = "VirtualCamera")
+	int32 HitProxySize = 5;
+
+	/** Components on these actors should not be considered. */
+	UPROPERTY(EditAnywhere, BlueprintreadWrite, Category = "VirtualCamera")
+	TArray<AActor*> IgnoredActors;
+};
+
+USTRUCT(BlueprintType)
+struct VIRTUALCAMERA_API FVCamTraceHitProxyResult
+{
+	GENERATED_BODY()
+
+	UPROPERTY(BlueprintReadOnly, Category = "VirtualCamera")
+	TWeakObjectPtr<AActor> HitActor;
+
+	UPROPERTY(BlueprintReadOnly, Category = "VirtualCamera")
+	TWeakObjectPtr<UPrimitiveComponent> HitComponent;
+
+	friend bool operator==(const FVCamTraceHitProxyResult& Left, const FVCamTraceHitProxyResult& Right)
+	{
+		return Left.HitActor == Right.HitActor && Left.HitComponent == Right.HitComponent;
+	}
+	friend bool operator!=(const FVCamTraceHitProxyResult& Left, const FVCamTraceHitProxyResult& Right)
+	{
+		return !(Left == Right);
+	}
+};
+
 UCLASS(config=VirtualCamera, BlueprintType)
 class VIRTUALCAMERA_API UVCamBlueprintFunctionLibrary : public UBlueprintFunctionLibrary
 {
@@ -201,6 +238,23 @@ public:
 	UFUNCTION(BlueprintCallable, Category = "VirtualCamera")
 	static bool DeprojectScreenToWorldByViewport(const FVector2D& InScreenPosition, EVCamTargetViewportID TargetViewport, FVector& OutWorldPosition, FVector& OutWorldDirection);
 
+	/**
+	 * Traces from the viewport and returns all components that contribute to the pixels surrounding InScreenPosition.
+	 * The size of the pixel area checked is controlled by InQueryParams.HitProxySize.
+	 * 
+	 * This finds actors that have NoCollision set. The actor is found by determining which actors contribute to the specified pixel.
+	 * This function is designed for Editor builds; in Runtime builds, it returns false.
+	 *
+	 * @param InScreenPosition The viewport position to trace
+	 * @param InTargetViewport The viewport to trace in
+	 * @param InQueryParams Parameters for how the actors should be queried
+	 * @param Result The result, set if this function returns true.
+	 *
+	 * @return Whether Result was written to
+	 */
+	UFUNCTION(BlueprintPure, Category = "VirtualCamera")
+	static bool MultiTraceHitProxyOnViewport(const FVector2D& InScreenPosition, EVCamTargetViewportID InTargetViewport, FVCamTraceHitProxyQueryParams InQueryParams, TArray<FVCamTraceHitProxyResult>& Result);
+
 private:
 
 #if WITH_EDITOR
