Remove old gameplay code from VirtualCamera.
No formal deprecation is required because these plugins are experimental.

diff --git a/Engine/Plugins/Experimental/VirtualCamera/Source/VirtualCamera/Private/ConcertVirtualCamera.cpp b/Engine/Plugins/Experimental/VirtualCamera/Source/VirtualCamera/Private/ConcertVirtualCamera.cpp
deleted file mode 100644
diff --git a/Engine/Plugins/Experimental/VirtualCamera/Source/VirtualCamera/Private/ConcertVirtualCamera.h b/Engine/Plugins/Experimental/VirtualCamera/Source/VirtualCamera/Private/ConcertVirtualCamera.h
deleted file mode 100644
diff --git a/Engine/Plugins/Experimental/VirtualCamera/Source/VirtualCamera/Private/FunctionLibraries/VCamBlueprintFunctionLibrary.cpp b/Engine/Plugins/Experimental/VirtualCamera/Source/VirtualCamera/Private/FunctionLibraries/VCamBlueprintFunctionLibrary.cpp
--- a/Engine/Plugins/Experimental/VirtualCamera/Source/VirtualCamera/Private/FunctionLibraries/VCamBlueprintFunctionLibrary.cpp
+++ b/Engine/Plugins/Experimental/VirtualCamera/Source/VirtualCamera/Private/FunctionLibraries/VCamBlueprintFunctionLibrary.cpp
@@ -2,9 +2,6 @@
 
 #include "FunctionLibraries/VCamBlueprintFunctionLibrary.h"
 
-#include "LevelSequence/VirtualCameraClipsMetaData.h"
-#include "VirtualCameraUserSettings.h"
-
 #include "AssetRegistry/AssetData.h"
 #include "CineCameraActor.h"
 #include "CineCameraComponent.h"
@@ -51,11 +48,6 @@ bool UVCamBlueprintFunctionLibrary::IsGameRunning()
 #endif
 }
 
-UVirtualCameraUserSettings* UVCamBlueprintFunctionLibrary::GetUserSettings()
-{
-	return GetMutableDefault<UVirtualCameraUserSettings>();
-}
-
 ULevelSequence* UVCamBlueprintFunctionLibrary::GetCurrentLevelSequence()
 {
 #if WITH_EDITOR
@@ -129,7 +121,6 @@ int32 UVCamBlueprintFunctionLibrary::GetLevelSequenceLengthInFrames(const ULevel
 	return 0;
 }
 
-
 int32 UVCamBlueprintFunctionLibrary::TimecodeToFrameAmount(FTimecode Timecode, const FFrameRate& InFrameRate)
 {
 	return Timecode.ToFrameNumber(InFrameRate).Value; 
@@ -160,7 +151,6 @@ bool UVCamBlueprintFunctionLibrary::IsCurrentLevelSequencePlaying()
 #endif
 }
 
-
 UTexture* UVCamBlueprintFunctionLibrary::ImportSnapshotTexture(FString FileName, FString SubFolderName, FString AbsolutePathPackage)
 {
 #if WITH_EDITOR
@@ -170,57 +160,24 @@ UTexture* UVCamBlueprintFunctionLibrary::ImportSnapshotTexture(FString FileName,
 #endif
 }
 
-
-bool UVCamBlueprintFunctionLibrary::ModifyLevelSequenceMetadata(UVirtualCameraClipsMetaData* LevelSequenceMetaData)
-{
-#if WITH_EDITOR
-	if (LevelSequenceMetaData)
-	{
-		LevelSequenceMetaData->MarkPackageDirty();
-
-		return UEditorAssetLibrary::SaveAsset(LevelSequenceMetaData->GetPathName());
-	}
-#endif
-
-	return false;
-}
-
-bool UVCamBlueprintFunctionLibrary::ModifyLevelSequenceMetadataForSelects(UVirtualCameraClipsMetaData* LevelSequenceMetaData, bool bIsSelected)
-{
-#if WITH_EDITOR
-	if (LevelSequenceMetaData)
-	{
-		LevelSequenceMetaData->SetSelected(bIsSelected);
-		LevelSequenceMetaData->MarkPackageDirty();
-
-		return UEditorAssetLibrary::SaveAsset(LevelSequenceMetaData->GetPathName());
-	}
-#endif
-
-	return false;
-
-}
-
 bool UVCamBlueprintFunctionLibrary::EditorSaveAsset(FString AssetPath)
 {
 #if WITH_EDITOR
 	return UEditorAssetLibrary::SaveAsset(AssetPath, true);
-#endif
+#else
 	return false;
-
+#endif
 }
 
-
 UObject* UVCamBlueprintFunctionLibrary::EditorLoadAsset(FString AssetPath)
 {
 #if WITH_EDITOR
 	return UEditorAssetLibrary::LoadAsset(AssetPath);
-#endif
+#else
 	return nullptr;
+#endif
 }
 
-
-
 void UVCamBlueprintFunctionLibrary::ModifyObjectMetadataTags(UObject* InObject, FName InTag, FString InValue)
 {
 #if WITH_EDITOR
diff --git a/Engine/Plugins/Experimental/VirtualCamera/Source/VirtualCamera/Private/LevelSequence/LevelSequencePlaybackController.cpp b/Engine/Plugins/Experimental/VirtualCamera/Source/VirtualCamera/Private/LevelSequence/LevelSequencePlaybackController.cpp
deleted file mode 100644
diff --git a/Engine/Plugins/Experimental/VirtualCamera/Source/VirtualCamera/Private/VirtualCamera.cpp b/Engine/Plugins/Experimental/VirtualCamera/Source/VirtualCamera/Private/VirtualCamera.cpp
--- a/Engine/Plugins/Experimental/VirtualCamera/Source/VirtualCamera/Private/VirtualCamera.cpp
+++ b/Engine/Plugins/Experimental/VirtualCamera/Source/VirtualCamera/Private/VirtualCamera.cpp
@@ -3,47 +3,19 @@
 #include "VirtualCamera.h"
 
 #include "AdvancedWidgetsModule.h"
-#include "ConcertVirtualCamera.h"
-
 
 DEFINE_LOG_CATEGORY(LogVirtualCamera);
 
-
-const IVirtualCameraModule& IVirtualCameraModule::Get()
-{
-	static const FName ModuleName = TEXT("VirtualCamera");
-	return FModuleManager::Get().GetModuleChecked<IVirtualCameraModule>(ModuleName);
-}
-
-
-/**
- *
- */
-class FVirtualCameraModuleImpl : public IVirtualCameraModule
+namespace UE::VirtualCamera
 {
-public:
-	virtual FConcertVirtualCameraManager* GetConcertVirtualCameraManager() const override
+	void FVirtualCameraModuleImpl::StartupModule()
 	{
-		return ConcertManager.Get();
-	}
-
-private:
-	virtual void StartupModule() override
-	{
-		LLM_SCOPE_BYNAME(TEXT("VirtualCamera"));
-		ConcertManager = MakeUnique<FConcertVirtualCameraManager>();
-
 		// Loads widgets (ex. RadialSlider) that are potentially referenced by assets
 		FModuleManager::Get().LoadModuleChecked<FAdvancedWidgetsModule>("AdvancedWidgets");
 	}
 
-	virtual void ShutdownModule() override
-	{
-		ConcertManager.Reset();
-	}
-
-	TUniquePtr<FConcertVirtualCameraManager> ConcertManager;
-};
-
+	void FVirtualCameraModuleImpl::ShutdownModule()
+	{}
+}
 
-IMPLEMENT_MODULE(FVirtualCameraModuleImpl, VirtualCamera)
\ No newline at end of file
+IMPLEMENT_MODULE(UE::VirtualCamera::FVirtualCameraModuleImpl, VirtualCamera)
\ No newline at end of file
diff --git a/Engine/Plugins/Experimental/VirtualCamera/Source/VirtualCamera/Private/VirtualCamera.h b/Engine/Plugins/Experimental/VirtualCamera/Source/VirtualCamera/Private/VirtualCamera.h
--- a/Engine/Plugins/Experimental/VirtualCamera/Source/VirtualCamera/Private/VirtualCamera.h
+++ b/Engine/Plugins/Experimental/VirtualCamera/Source/VirtualCamera/Private/VirtualCamera.h
@@ -2,22 +2,20 @@
 
 #pragma once
 
-#include "CoreMinimal.h"
 #include "Modules/ModuleInterface.h"
 #include "Modules/ModuleManager.h"
 
 DECLARE_LOG_CATEGORY_EXTERN(LogVirtualCamera, Log, Log);
 
-
-class FConcertVirtualCameraManager;
-
-/**
- *
- */
-class IVirtualCameraModule : public IModuleInterface
+namespace UE::VirtualCamera
 {
-public:
-	static const IVirtualCameraModule& Get();
+	class FVirtualCameraModuleImpl : public IModuleInterface
+	{
+	public:
 
-	virtual FConcertVirtualCameraManager* GetConcertVirtualCameraManager() const = 0;
-};
+		//~ Begin IModuleInterface Interface
+		virtual void StartupModule() override;
+		virtual void ShutdownModule() override;
+		//~ End IModuleInterface Interface
+	};
+}
diff --git a/Engine/Plugins/Experimental/VirtualCamera/Source/VirtualCamera/Private/VirtualCameraActor.cpp b/Engine/Plugins/Experimental/VirtualCamera/Source/VirtualCamera/Private/VirtualCameraActor.cpp
deleted file mode 100644
diff --git a/Engine/Plugins/Experimental/VirtualCamera/Source/VirtualCamera/Private/VirtualCameraCineCameraComponent.cpp b/Engine/Plugins/Experimental/VirtualCamera/Source/VirtualCamera/Private/VirtualCameraCineCameraComponent.cpp
deleted file mode 100644
diff --git a/Engine/Plugins/Experimental/VirtualCamera/Source/VirtualCamera/Private/VirtualCameraConcertCameraComponent.cpp b/Engine/Plugins/Experimental/VirtualCamera/Source/VirtualCamera/Private/VirtualCameraConcertCameraComponent.cpp
deleted file mode 100644
diff --git a/Engine/Plugins/Experimental/VirtualCamera/Source/VirtualCamera/Private/VirtualCameraMovement.cpp b/Engine/Plugins/Experimental/VirtualCamera/Source/VirtualCamera/Private/VirtualCameraMovement.cpp
deleted file mode 100644
diff --git a/Engine/Plugins/Experimental/VirtualCamera/Source/VirtualCamera/Private/VirtualCameraMovementComponent.cpp b/Engine/Plugins/Experimental/VirtualCamera/Source/VirtualCamera/Private/VirtualCameraMovementComponent.cpp
deleted file mode 100644
diff --git a/Engine/Plugins/Experimental/VirtualCamera/Source/VirtualCamera/Private/VirtualCameraPawnBase.cpp b/Engine/Plugins/Experimental/VirtualCamera/Source/VirtualCamera/Private/VirtualCameraPawnBase.cpp
deleted file mode 100644
diff --git a/Engine/Plugins/Experimental/VirtualCamera/Source/VirtualCamera/Private/VirtualCameraPlayerControllerBase.cpp b/Engine/Plugins/Experimental/VirtualCamera/Source/VirtualCamera/Private/VirtualCameraPlayerControllerBase.cpp
deleted file mode 100644
diff --git a/Engine/Plugins/Experimental/VirtualCamera/Source/VirtualCamera/Private/VirtualCameraSaveGame.cpp b/Engine/Plugins/Experimental/VirtualCamera/Source/VirtualCamera/Private/VirtualCameraSaveGame.cpp
deleted file mode 100644
diff --git a/Engine/Plugins/Experimental/VirtualCamera/Source/VirtualCamera/Private/VirtualCameraSubsystem.cpp b/Engine/Plugins/Experimental/VirtualCamera/Source/VirtualCamera/Private/VirtualCameraSubsystem.cpp
deleted file mode 100644
diff --git a/Engine/Plugins/Experimental/VirtualCamera/Source/VirtualCamera/Private/VirtualCameraUserSettings.cpp b/Engine/Plugins/Experimental/VirtualCamera/Source/VirtualCamera/Private/VirtualCameraUserSettings.cpp
deleted file mode 100644
diff --git a/Engine/Plugins/Experimental/VirtualCamera/Source/VirtualCamera/Public/FunctionLibraries/VCamBlueprintFunctionLibrary.h b/Engine/Plugins/Experimental/VirtualCamera/Source/VirtualCamera/Public/FunctionLibraries/VCamBlueprintFunctionLibrary.h
--- a/Engine/Plugins/Experimental/VirtualCamera/Source/VirtualCamera/Public/FunctionLibraries/VCamBlueprintFunctionLibrary.h
+++ b/Engine/Plugins/Experimental/VirtualCamera/Source/VirtualCamera/Public/FunctionLibraries/VCamBlueprintFunctionLibrary.h
@@ -17,7 +17,6 @@ class ULevelSequence;
 class UPrimitiveComponent;
 class USceneCaptureComponent2D;
 class UVirtualCameraClipsMetaData;
-class UVirtualCameraUserSettings; 
 
 #if WITH_EDITOR
 class ISequencer;
@@ -60,7 +59,7 @@ struct VIRTUALCAMERA_API FVCamTraceHitProxyResult
 	}
 };
 
-UCLASS(config=VirtualCamera, BlueprintType)
+UCLASS(BlueprintType)
 class VIRTUALCAMERA_API UVCamBlueprintFunctionLibrary : public UBlueprintFunctionLibrary
 {
 	GENERATED_BODY()
@@ -71,9 +70,6 @@ public:
 	UFUNCTION(BlueprintPure, Category = "VirtualCamera")
 	static bool IsGameRunning();
 
-	UFUNCTION(BlueprintPure, Category = "VirtualCamera")
-	static UVirtualCameraUserSettings* GetUserSettings();
-
 	/**
 	 * Get the currently opened level sequence asset
 	 * @see ULevelSequenceEditorBlueprintLibrary::GetCurrentLevelSequence
@@ -127,14 +123,6 @@ public:
 	/** Imports image as a uasset */
 	UFUNCTION(BlueprintCallable, Category = "VirtualCamera")
 	static UTexture* ImportSnapshotTexture(FString FileName, FString SubFolderName, FString AbsolutePathPackage);
-
-	/** Saves UVirtualCameraClipsMetaData with updated selects information. */ 
-	UFUNCTION(BlueprintCallable, Category = "VirtualCamera | Metadata")
-	static bool ModifyLevelSequenceMetadataForSelects(UVirtualCameraClipsMetaData* LevelSequenceMetaData, bool bIsSelected);
-	
-	/** Marks a LevelSequence as dirty and saves it, persisting metadata changes */
-	UFUNCTION(BlueprintCallable, Category = "VirtualCamera | Metadata")
-	static bool ModifyLevelSequenceMetadata(UVirtualCameraClipsMetaData* LevelSequenceMetaData);
 	
 	/** Save an asset through path. Returns true on success. */
 	UFUNCTION(BlueprintCallable, Category = "VirtualCamera | Clips")
diff --git a/Engine/Plugins/Experimental/VirtualCamera/Source/VirtualCamera/Public/IVirtualCameraController.h b/Engine/Plugins/Experimental/VirtualCamera/Source/VirtualCamera/Public/IVirtualCameraController.h
deleted file mode 100644
diff --git a/Engine/Plugins/Experimental/VirtualCamera/Source/VirtualCamera/Public/IVirtualCameraOptions.h b/Engine/Plugins/Experimental/VirtualCamera/Source/VirtualCamera/Public/IVirtualCameraOptions.h
deleted file mode 100644
diff --git a/Engine/Plugins/Experimental/VirtualCamera/Source/VirtualCamera/Public/IVirtualCameraPresetContainer.h b/Engine/Plugins/Experimental/VirtualCamera/Source/VirtualCamera/Public/IVirtualCameraPresetContainer.h
deleted file mode 100644
diff --git a/Engine/Plugins/Experimental/VirtualCamera/Source/VirtualCamera/Public/LevelSequence/LevelSequencePlaybackController.h b/Engine/Plugins/Experimental/VirtualCamera/Source/VirtualCamera/Public/LevelSequence/LevelSequencePlaybackController.h
deleted file mode 100644
diff --git a/Engine/Plugins/Experimental/VirtualCamera/Source/VirtualCamera/Public/VirtualCameraActor.h b/Engine/Plugins/Experimental/VirtualCamera/Source/VirtualCamera/Public/VirtualCameraActor.h
deleted file mode 100644
diff --git a/Engine/Plugins/Experimental/VirtualCamera/Source/VirtualCamera/Public/VirtualCameraCineCameraComponent.h b/Engine/Plugins/Experimental/VirtualCamera/Source/VirtualCamera/Public/VirtualCameraCineCameraComponent.h
deleted file mode 100644
diff --git a/Engine/Plugins/Experimental/VirtualCamera/Source/VirtualCamera/Public/VirtualCameraConcertCameraComponent.h b/Engine/Plugins/Experimental/VirtualCamera/Source/VirtualCamera/Public/VirtualCameraConcertCameraComponent.h
deleted file mode 100644
diff --git a/Engine/Plugins/Experimental/VirtualCamera/Source/VirtualCamera/Public/VirtualCameraMovement.h b/Engine/Plugins/Experimental/VirtualCamera/Source/VirtualCamera/Public/VirtualCameraMovement.h
deleted file mode 100644
diff --git a/Engine/Plugins/Experimental/VirtualCamera/Source/VirtualCamera/Public/VirtualCameraMovementComponent.h b/Engine/Plugins/Experimental/VirtualCamera/Source/VirtualCamera/Public/VirtualCameraMovementComponent.h
deleted file mode 100644
diff --git a/Engine/Plugins/Experimental/VirtualCamera/Source/VirtualCamera/Public/VirtualCameraPawnBase.h b/Engine/Plugins/Experimental/VirtualCamera/Source/VirtualCamera/Public/VirtualCameraPawnBase.h
deleted file mode 100644
diff --git a/Engine/Plugins/Experimental/VirtualCamera/Source/VirtualCamera/Public/VirtualCameraPlayerControllerBase.h b/Engine/Plugins/Experimental/VirtualCamera/Source/VirtualCamera/Public/VirtualCameraPlayerControllerBase.h
deleted file mode 100644
diff --git a/Engine/Plugins/Experimental/VirtualCamera/Source/VirtualCamera/Public/VirtualCameraSaveGame.h b/Engine/Plugins/Experimental/VirtualCamera/Source/VirtualCamera/Public/VirtualCameraSaveGame.h
deleted file mode 100644
diff --git a/Engine/Plugins/Experimental/VirtualCamera/Source/VirtualCamera/Public/VirtualCameraSubsystem.h b/Engine/Plugins/Experimental/VirtualCamera/Source/VirtualCamera/Public/VirtualCameraSubsystem.h
deleted file mode 100644
diff --git a/Engine/Plugins/Experimental/VirtualCamera/Source/VirtualCamera/Public/VirtualCameraUserSettings.h b/Engine/Plugins/Experimental/VirtualCamera/Source/VirtualCamera/Public/VirtualCameraUserSettings.h
deleted file mode 100644
diff --git a/Engine/Plugins/Experimental/VirtualCamera/Source/VirtualCamera/Public/VirtualCameraVersion.h b/Engine/Plugins/Experimental/VirtualCamera/Source/VirtualCamera/Public/VirtualCameraVersion.h
deleted file mode 100644
diff --git a/Engine/Plugins/Experimental/VirtualCamera/Source/VirtualCamera/VirtualCamera.Build.cs b/Engine/Plugins/Experimental/VirtualCamera/Source/VirtualCamera/VirtualCamera.Build.cs
--- a/Engine/Plugins/Experimental/VirtualCamera/Source/VirtualCamera/VirtualCamera.Build.cs
+++ b/Engine/Plugins/Experimental/VirtualCamera/Source/VirtualCamera/VirtualCamera.Build.cs
@@ -37,25 +37,7 @@ public class VirtualCamera : ModuleRules
 				"AdvancedWidgets"
 			}
 		);
-
-		if (Target.Type == TargetType.Editor || Target.Type == TargetType.Program)
-		{
-			PrivateDefinitions.Add("VIRTUALCAMERA_WITH_CONCERT=1");
-			PrivateDependencyModuleNames.AddRange(
-				new string[]
-				{
-					"Concert",
-					"ConcertSyncClient",
-					"MultiUserClient",
-				}
-			);
-		}
-		else
-		{
-			PrivateDefinitions.Add("VIRTUALCAMERA_WITH_CONCERT=0");
-		}
-
-
+		
 		if (Target.bBuildEditor == true)
 		{
 			PublicDependencyModuleNames.Add("LevelSequenceEditor");
diff --git a/Engine/Plugins/Experimental/VirtualCamera/Source/VirtualCameraEditor/Private/VirtualCameraEditorModule.cpp b/Engine/Plugins/Experimental/VirtualCamera/Source/VirtualCameraEditor/Private/VirtualCameraEditorModule.cpp
--- a/Engine/Plugins/Experimental/VirtualCamera/Source/VirtualCameraEditor/Private/VirtualCameraEditorModule.cpp
+++ b/Engine/Plugins/Experimental/VirtualCamera/Source/VirtualCameraEditor/Private/VirtualCameraEditorModule.cpp
@@ -1,19 +1,14 @@
 // Copyright Epic Games, Inc. All Rights Reserved.
 
-#include "ISettingsModule.h"
-#include "ISettingsSection.h"
 #include "Modules/ModuleManager.h"
-#include "VirtualCameraUserSettings.h"
 
 #include "ActorFactories/ActorFactoryBlueprint.h"
 #include "CineCameraActor.h"
 #include "Filters/CustomClassFilterData.h"
-#include "Interfaces/IPluginManager.h"
 #include "IPlacementModeModule.h"
 #include "IVPUtilitiesEditorModule.h"
 #include "LevelEditor.h"
 #include "LevelEditorOutlinerSettings.h"
-#include "VirtualCameraActor.h"
 
 #define LOCTEXT_NAMESPACE "FVirtualCameraEditorModule"
 
@@ -24,15 +19,12 @@ class FVirtualCameraEditorModule : public IModuleInterface
 public:
 	virtual void StartupModule() override
 	{
-		RegisterSettings();
 		RegisterPlacementModeItems();
 		RegisterOutlinerFilters();
 	}
 
 	virtual void ShutdownModule() override
-	{
-		UnregisterSettings();
-	}
+	{}
 
 private:
 	
@@ -77,34 +69,9 @@ private:
 				TSharedRef<FCustomClassFilterData> CineCameraActorClassData =
 					MakeShared<FCustomClassFilterData>(ACineCameraActor::StaticClass(), VPFilterCategory, FLinearColor::White);
 				LevelEditorModule->AddCustomClassFilterToOutliner(CineCameraActorClassData);
-
-				TSharedRef<FCustomClassFilterData> VirtualCameraActorClassData =
-					MakeShared<FCustomClassFilterData>(AVirtualCameraActor::StaticClass(), VPFilterCategory, FLinearColor::White);
-				LevelEditorModule->AddCustomClassFilterToOutliner(VirtualCameraActorClassData);
 			}
 		}
 	}
-	
-	void RegisterSettings()
-	{
-		ISettingsModule* SettingsModule = FModuleManager::GetModulePtr<ISettingsModule>("Settings");
-		if (SettingsModule != nullptr)
-		{
-			ISettingsSectionPtr SettingsSection = SettingsModule->RegisterSettings("Project", "Plugins", "VirtualCamera",
-				LOCTEXT("VirtualCameraUserSettingsName", "Virtual Camera"),
-				LOCTEXT("VirtualCameraUserSettingsDescription", "Configure the Virtual Camera settings."),
-				GetMutableDefault<UVirtualCameraUserSettings>());
-		}
-	}
-
-	void UnregisterSettings()
-	{
-		ISettingsModule* SettingsModule = FModuleManager::GetModulePtr<ISettingsModule>("Settings");
-		if (SettingsModule != nullptr)
-		{
-			SettingsModule->UnregisterSettings("Project", "Plugins", "VirtualCamera");
-		}
-	}
 };
 #undef LOCTEXT_NAMESPACE
 
