Move VirtualCamera and VirtualCameraCore from Plugins/Experimental Plugins/VirtualProduction

#virtualized

diff --git a/Engine/Plugins/VirtualProduction/VirtualCamera/Source/VirtualCamera/Private/LevelSequence/VirtualCameraClipsMetaData.cpp b/Engine/Plugins/VirtualProduction/VirtualCamera/Source/VirtualCamera/Private/LevelSequence/VirtualCameraClipsMetaData.cpp
new file mode 100644
--- /dev/null
+++ b/Engine/Plugins/VirtualProduction/VirtualCamera/Source/VirtualCamera/Private/LevelSequence/VirtualCameraClipsMetaData.cpp
@@ -0,0 +1,189 @@
+// Copyright Epic Games, Inc. All Rights Reserved.
+
+#include "LevelSequence/VirtualCameraClipsMetaData.h"
+
+#include "UObject/AssetRegistryTagsContext.h"
+
+const FName UVirtualCameraClipsMetaData::AssetRegistryTag_FocalLength = "ClipsMetaData_FocalLength";
+const FName UVirtualCameraClipsMetaData::AssetRegistryTag_bIsSelected = "ClipsMetaData_bIsSelected";
+const FName UVirtualCameraClipsMetaData::AssetRegistryTag_RecordedLevelName = "ClipsMetaData_RecordedLevelName";
+const FName UVirtualCameraClipsMetaData::AssetRegistryTag_FrameCountStart = "ClipsMetaData_FrameCountStart";
+const FName UVirtualCameraClipsMetaData::AssetRegistryTag_FrameCountEnd = "ClipsMetaData_FrameCountEnd"; 
+const FName UVirtualCameraClipsMetaData::AssetRegistryTag_LengthInFrames = "ClipsMetaData_LengthInFrames";
+const FName UVirtualCameraClipsMetaData::AssetRegistryTag_DisplayRate = "ClipsMetaData_DisplayRate";
+const FName UVirtualCameraClipsMetaData::AssetRegistryTag_bIsACineCameraRecording = "ClipsMetaData_bIsACineCameraRecording";
+const FName UVirtualCameraClipsMetaData::AssetRegistryTag_bIsNoGood = "ClipsMetaData_bIsNoGood";
+const FName UVirtualCameraClipsMetaData::AssetRegistryTag_bIsFlagged = "ClipsMetaData_bIsFlagged";
+const FName UVirtualCameraClipsMetaData::AssetRegistryTag_FavoriteLevel = "ClipsMetaData_FavoriteLevel";
+const FName UVirtualCameraClipsMetaData::AssetRegistryTag_bIsCreatedFromVCam = "ClipsMetaData_bIsCreatedFromVCam";
+
+UVirtualCameraClipsMetaData::UVirtualCameraClipsMetaData(const FObjectInitializer& ObjInit)
+	: Super(ObjInit)
+{
+	FocalLength = 0;
+	bIsSelected = false;
+	bIsACineCameraRecording = false; 
+}
+
+TSet<FName> UVirtualCameraClipsMetaData::GetAllClipsMetaDataTags()
+{
+	return {
+		GetClipsMetaDataTag_FocalLength(),
+		GetClipsMetaDataTag_IsSelected(),
+		GetClipsMetaDataTag_RecordedLevel(),
+		GetClipsMetaDataTag_FrameCountStart(),
+		GetClipsMetaDataTag_FrameCountEnd(),
+		GetClipsMetaDataTag_LengthInFrames(),
+		GetClipsMetaDataTag_DisplayRate(),
+		GetClipsMetaDataTag_IsCineACineCameraRecording(),
+		GetClipsMetaDataTag_IsNoGood(),
+		GetClipsMetaDataTag_IsFlagged(),
+		GetClipsMetaDataTag_FavoriteLevel(),
+		GetClipsMetaDataTag_IsCreatedFromVCam()
+	};
+}
+
+void UVirtualCameraClipsMetaData::ExtendAssetRegistryTags(FAssetRegistryTagsContext Context) const
+{
+	IMovieSceneMetaDataInterface::ExtendAssetRegistryTags(Context);
+
+	FString IsSelected = bIsSelected ? "true" : "false";
+	FString IsRecordedFromACineCameraActor = bIsACineCameraRecording ? "true" : "false";
+
+	Context.AddTag(FAssetRegistryTag(AssetRegistryTag_FocalLength, FString::FromInt(FocalLength), FAssetRegistryTag::ETagType::TT_Numerical, FAssetRegistryTag::TD_None));
+	Context.AddTag(FAssetRegistryTag(AssetRegistryTag_bIsSelected, IsSelected, FAssetRegistryTag::ETagType::TT_Alphabetical, FAssetRegistryTag::TD_None));
+	Context.AddTag(FAssetRegistryTag(AssetRegistryTag_RecordedLevelName, RecordedLevelName, FAssetRegistryTag::ETagType::TT_Alphabetical, FAssetRegistryTag::TD_None));
+	Context.AddTag(FAssetRegistryTag(AssetRegistryTag_FrameCountStart, FString::FromInt(FrameCountStart), FAssetRegistryTag::ETagType::TT_Numerical, FAssetRegistryTag::TD_None));
+	Context.AddTag(FAssetRegistryTag(AssetRegistryTag_FrameCountEnd, FString::FromInt(FrameCountEnd), FAssetRegistryTag::ETagType::TT_Numerical, FAssetRegistryTag::TD_None));
+	Context.AddTag(FAssetRegistryTag(AssetRegistryTag_LengthInFrames, FString::FromInt(LengthInFrames), FAssetRegistryTag::ETagType::TT_Numerical, FAssetRegistryTag::TD_None));
+	Context.AddTag(FAssetRegistryTag(AssetRegistryTag_DisplayRate, DisplayRate.ToPrettyText().ToString(), FAssetRegistryTag::ETagType::TT_Alphabetical, FAssetRegistryTag::TD_None));
+	Context.AddTag(FAssetRegistryTag(AssetRegistryTag_bIsACineCameraRecording, IsRecordedFromACineCameraActor, FAssetRegistryTag::ETagType::TT_Alphabetical, FAssetRegistryTag::TD_None));
+	Context.AddTag(FAssetRegistryTag(AssetRegistryTag_bIsNoGood, FString::FromInt(bIsNoGood), FAssetRegistryTag::ETagType::TT_Numerical, FAssetRegistryTag::TD_None));
+	Context.AddTag(FAssetRegistryTag(AssetRegistryTag_bIsFlagged, FString::FromInt(bIsFlagged), FAssetRegistryTag::ETagType::TT_Numerical, FAssetRegistryTag::TD_None));
+	Context.AddTag(FAssetRegistryTag(AssetRegistryTag_FavoriteLevel, FString::FromInt(FavoriteLevel), FAssetRegistryTag::ETagType::TT_Alphabetical, FAssetRegistryTag::TD_None));
+	Context.AddTag(FAssetRegistryTag(AssetRegistryTag_bIsCreatedFromVCam, FString::FromInt(bIsCreatedFromVCam), FAssetRegistryTag::ETagType::TT_Numerical, FAssetRegistryTag::TD_None));
+}
+
+#if WITH_EDITOR
+void UVirtualCameraClipsMetaData::ExtendAssetRegistryTagMetaData(TMap<FName, FAssetRegistryTagMetadata>& OutMetadata) const
+{
+	OutMetadata.Add(AssetRegistryTag_FocalLength, FAssetRegistryTagMetadata()
+		.SetDisplayName(NSLOCTEXT("ClipsMetaData", "FocalLength_Label", "FocalLength"))
+		.SetTooltip(NSLOCTEXT("ClipsMetaData", "FocalLength_Tip", "The focal length that this level sequence was recorded with"))
+	);
+
+	OutMetadata.Add(AssetRegistryTag_bIsSelected, FAssetRegistryTagMetadata()
+		.SetDisplayName(NSLOCTEXT("ClipsMetaData", "IsSelected_Label", "bIsSelected"))
+		.SetTooltip(NSLOCTEXT("ClipsMetaData", "IsSelected_Tip", "If this level sequence is marked as a selected take"))
+	);
+
+	OutMetadata.Add(AssetRegistryTag_RecordedLevelName, FAssetRegistryTagMetadata()
+		.SetDisplayName(NSLOCTEXT("ClipsMetaData", "RecordedLavelName_Label", "RecordedLevelName"))
+		.SetTooltip(NSLOCTEXT("ClipsMetaData", "RecordedLavelName_Tip", "The name of the level this sequence was recorded"))
+	);
+
+	OutMetadata.Add(AssetRegistryTag_FrameCountStart, FAssetRegistryTagMetadata()
+		.SetDisplayName(NSLOCTEXT("ClipsMetaData", "FrameCountStart_Label", "FrameCountStart"))
+		.SetTooltip(NSLOCTEXT("ClipsMetaData", "FrameCountStart_Tip", "The first frame of the level sequence. "))
+	);
+
+	OutMetadata.Add(AssetRegistryTag_FrameCountEnd, FAssetRegistryTagMetadata()
+		.SetDisplayName(NSLOCTEXT("ClipsMetaData", "FrameCountEnd_Label", "FrameCountEnd"))
+		.SetTooltip(NSLOCTEXT("ClipsMetaData", "FrameCountEnd_Tip", "The last frame of the level sequence."))
+	);
+	OutMetadata.Add(AssetRegistryTag_LengthInFrames, FAssetRegistryTagMetadata()
+		.SetDisplayName(NSLOCTEXT("ClipsMetaData", "LengthInFrames_Label", "LengthInFrames"))
+		.SetTooltip(NSLOCTEXT("ClipsMetaData", "LengthInFrames_Tip", "The length in frames of the level sequence."))
+	);
+	OutMetadata.Add(AssetRegistryTag_DisplayRate, FAssetRegistryTagMetadata()
+		.SetDisplayName(NSLOCTEXT("ClipsMetaData", "DisplayRate_Label", "DisplayRate"))
+		.SetTooltip(NSLOCTEXT("ClipsMetaData", "DisplayRate_Tip", "The display rate of level sequence."))
+	);
+
+	OutMetadata.Add(AssetRegistryTag_bIsACineCameraRecording, FAssetRegistryTagMetadata()
+		.SetDisplayName(NSLOCTEXT("ClipsMetaData", "IsACineCameraRecording_Label", "bIsACineCameraRecording"))
+		.SetTooltip(NSLOCTEXT("ClipsMetaData", "IsACineCameraRecording_Tip", "If the level sequence was recorded with a CineCameraActor."))
+	);
+}
+#endif
+
+float UVirtualCameraClipsMetaData::GetFocalLength() const 
+{
+	return FocalLength; 
+}
+
+bool UVirtualCameraClipsMetaData::GetSelected() const
+{
+	return bIsSelected; 
+}
+
+FString UVirtualCameraClipsMetaData::GetRecordedLevelName() const
+{
+	return RecordedLevelName; 
+}
+
+int UVirtualCameraClipsMetaData::GetFrameCountStart() const
+{
+	return FrameCountStart; 
+}
+
+int UVirtualCameraClipsMetaData::GetFrameCountEnd() const
+{
+	return FrameCountEnd; 
+}
+
+int UVirtualCameraClipsMetaData::GetLengthInFrames()
+{
+	return LengthInFrames;
+}
+
+FFrameRate UVirtualCameraClipsMetaData::GetDisplayRate()
+{
+	return DisplayRate;
+}
+
+bool UVirtualCameraClipsMetaData::GetIsACineCameraRecording() const
+{
+	return bIsACineCameraRecording; 
+}
+
+void UVirtualCameraClipsMetaData::SetFocalLength(float InFocalLength) 
+{
+	FocalLength = InFocalLength;
+}
+
+void UVirtualCameraClipsMetaData::SetSelected(bool bInSelected)
+{
+	bIsSelected = bInSelected;
+}
+
+void UVirtualCameraClipsMetaData::SetRecordedLevelName(FString InLevelName)
+{
+	RecordedLevelName = InLevelName;
+}
+
+void UVirtualCameraClipsMetaData::SetFrameCountStart(int InFrame)
+{
+	FrameCountStart = InFrame; 
+}
+
+void UVirtualCameraClipsMetaData::SetFrameCountEnd(int InFrame)
+{
+	FrameCountEnd = InFrame; 
+}
+
+void UVirtualCameraClipsMetaData::SetLengthInFrames(int InLength)
+{
+	LengthInFrames = InLength; 
+}
+
+void UVirtualCameraClipsMetaData::SetDisplayRate(FFrameRate InDisplayRate)
+{
+	DisplayRate = InDisplayRate; 
+}
+
+void UVirtualCameraClipsMetaData::SetIsACineCameraRecording(bool bInIsACineCameraRecording)
+{
+	bIsACineCameraRecording = bInIsACineCameraRecording; 
+}
+
diff --git a/Engine/Plugins/VirtualProduction/VirtualCamera/Source/VirtualCamera/VirtualCamera.Build.cs b/Engine/Plugins/VirtualProduction/VirtualCamera/Source/VirtualCamera/VirtualCamera.Build.cs
new file mode 100644
--- /dev/null
+++ b/Engine/Plugins/VirtualProduction/VirtualCamera/Source/VirtualCamera/VirtualCamera.Build.cs
@@ -0,0 +1,55 @@
+// Copyright Epic Games, Inc. All Rights Reserved.
+
+using UnrealBuildTool;
+
+public class VirtualCamera : ModuleRules
+{
+	public VirtualCamera(ReadOnlyTargetRules Target) : base(Target)
+	{
+		DefaultBuildSettings = BuildSettingsVersion.V2;
+
+		PublicDependencyModuleNames.AddRange(
+			new string[]
+			{
+				"AugmentedReality",
+				"CinematicCamera",
+				"Core",
+				"CoreUObject",
+				"Engine",
+				"InputCore",
+				"LevelSequence",
+				"LiveLinkInterface",
+				"MovieScene",
+				"RemoteSession",
+				"TimeManagement",
+				"UMG",
+				"VCamCore",
+				"VPUtilities",
+				"AssetRegistry",
+			}
+		);
+
+		PrivateDependencyModuleNames.AddRange(
+			new string[]
+			{
+				"MediaIOCore",
+				"Slate",
+				"AdvancedWidgets"
+			}
+		);
+		
+		if (Target.bBuildEditor == true)
+		{
+			PublicDependencyModuleNames.Add("LevelSequenceEditor");
+			PublicDependencyModuleNames.Add("Sequencer");
+			PublicDependencyModuleNames.Add("SlateCore");
+			PublicDependencyModuleNames.Add("TakeRecorder");
+			PrivateDependencyModuleNames.Add("LevelEditor");
+			PrivateDependencyModuleNames.Add("EditorFramework");
+			PrivateDependencyModuleNames.Add("UnrealEd");
+			PrivateDependencyModuleNames.Add("EditorScriptingUtilities");
+			PrivateDependencyModuleNames.Add("VPUtilitiesEditor");
+			PrivateDependencyModuleNames.Add("TakesCore");
+		}
+	}
+}
