Allow all platform to load VirtualCamera module

diff --git a/Engine/Plugins/VirtualProduction/VirtualCamera/VirtualCamera.uplugin b/Engine/Plugins/VirtualProduction/VirtualCamera/VirtualCamera.uplugin
--- a/Engine/Plugins/VirtualProduction/VirtualCamera/VirtualCamera.uplugin
+++ b/Engine/Plugins/VirtualProduction/VirtualCamera/VirtualCamera.uplugin
@@ -28,11 +28,7 @@
 		{
 			"Name": "VirtualCamera",
 			"Type": "Runtime",
-			"LoadingPhase": "Default",
-			"PlatformAllowList": [
-				"Win64",
-				"Linux"
-			]
+			"LoadingPhase": "Default"
 		},
 		{
 			"Name": "VirtualCameraEditor",
