Add UVirtualCameraUserSettings which saves per user the tutorial state

diff --git a/Engine/Plugins/VirtualProduction/VirtualCamera/Source/VirtualCamera/Private/VirtualCamera.cpp b/Engine/Plugins/VirtualProduction/VirtualCamera/Source/VirtualCamera/Private/VirtualCamera.cpp
--- a/Engine/Plugins/VirtualProduction/VirtualCamera/Source/VirtualCamera/Private/VirtualCamera.cpp
+++ b/Engine/Plugins/VirtualProduction/VirtualCamera/Source/VirtualCamera/Private/VirtualCamera.cpp
@@ -3,8 +3,15 @@
 #include "VirtualCamera.h"
 
 #include "AdvancedWidgetsModule.h"
+#include "Misc/CoreDelegates.h"
+
+#if WITH_EDITOR
+#include "ISettingsModule.h"
+#include "VCamUserSettings.h"
+#endif
 
 DEFINE_LOG_CATEGORY(LogVirtualCamera);
+#define LOCTEXT_NAMESPACE "FVirtualCameraModuleImpl"
 
 namespace UE::VirtualCamera
 {
@@ -12,10 +19,42 @@ namespace UE::VirtualCamera
 	{
 		// Loads widgets (ex. RadialSlider) that are potentially referenced by assets
 		FModuleManager::Get().LoadModuleChecked<FAdvancedWidgetsModule>("AdvancedWidgets");
+
+		FCoreDelegates::OnPostEngineInit.AddRaw(this, &FVirtualCameraModuleImpl::RegisterSettings);
 	}
 
 	void FVirtualCameraModuleImpl::ShutdownModule()
-	{}
+	{
+		UnregisterSettings();
+	}
+
+	void FVirtualCameraModuleImpl::RegisterSettings()
+	{
+#if WITH_EDITOR
+		ISettingsModule* SettingsModule = FModuleManager::GetModulePtr<ISettingsModule>("Settings");
+		if (ensure(SettingsModule))
+		{
+			SettingsModule->RegisterSettings("Project", "Plugins", "VirtualCamera",
+				LOCTEXT("VirtualCameraUserSettingsName", "Virtual Camera"),
+				LOCTEXT("VirtualCameraUserSettingsDescription", "Configure the Virtual Camera settings."),
+				GetMutableDefault<UVirtualCameraUserSettings>()
+				);
+		}
+#endif
+	}
+
+	void FVirtualCameraModuleImpl::UnregisterSettings()
+	{
+#if WITH_EDITOR
+		ISettingsModule* SettingsModule = FModuleManager::GetModulePtr<ISettingsModule>("Settings");
+		if (SettingsModule)
+		{
+			SettingsModule->UnregisterSettings("Project", "Plugins", "VirtualCamera");
+		}
+#endif
+	}
 }
 
-IMPLEMENT_MODULE(UE::VirtualCamera::FVirtualCameraModuleImpl, VirtualCamera)
\ No newline at end of file
+IMPLEMENT_MODULE(UE::VirtualCamera::FVirtualCameraModuleImpl, VirtualCamera)
+
+#undef LOCTEXT_NAMESPACE
\ No newline at end of file
diff --git a/Engine/Plugins/VirtualProduction/VirtualCamera/Source/VirtualCamera/Private/VirtualCamera.h b/Engine/Plugins/VirtualProduction/VirtualCamera/Source/VirtualCamera/Private/VirtualCamera.h
--- a/Engine/Plugins/VirtualProduction/VirtualCamera/Source/VirtualCamera/Private/VirtualCamera.h
+++ b/Engine/Plugins/VirtualProduction/VirtualCamera/Source/VirtualCamera/Private/VirtualCamera.h
@@ -17,5 +17,10 @@ namespace UE::VirtualCamera
 		virtual void StartupModule() override;
 		virtual void ShutdownModule() override;
 		//~ End IModuleInterface Interface
+
+	private:
+
+		void RegisterSettings();
+		void UnregisterSettings();
 	};
 }
diff --git a/Engine/Plugins/VirtualProduction/VirtualCamera/Source/VirtualCamera/Public/VCamUserSettings.h b/Engine/Plugins/VirtualProduction/VirtualCamera/Source/VirtualCamera/Public/VCamUserSettings.h
new file mode 100644
--- /dev/null
+++ b/Engine/Plugins/VirtualProduction/VirtualCamera/Source/VirtualCamera/Public/VCamUserSettings.h
@@ -0,0 +1,39 @@
+// Copyright Epic Games, Inc. All Rights Reserved.
+
+#pragma once
+
+#include "UObject/Object.h"
+#include "UObject/UObjectGlobals.h"
+#include "VCamUserSettings.generated.h"
+
+UENUM(BlueprintType)
+enum class EVCamTutorialCompletionState : uint8
+{
+	/** Tutorial not completed */
+	Pending,
+	/** Tutorial was completed */
+	Completed
+};
+
+UCLASS(BlueprintType, Config = EditorPerProjectUserSettings)
+class VIRTUALCAMERA_API UVirtualCameraUserSettings : public UObject
+{
+	GENERATED_BODY()
+
+public:
+
+	UFUNCTION(BlueprintPure, Category = "Virtual Camera", meta = (Keywords = "Get Settings Virtual Camera VCam"))
+	UVirtualCameraUserSettings* GetSettings() { return GetMutableDefault<UVirtualCameraUserSettings>(); }
+
+	/**
+	 * Indicates whether the VCam tutorial is completed.
+	 * You can manually reset this to Pending if you want to retake the tutorial.
+	 * The tutorial shown in the default VCamHUD, e.g. to teach gestures.
+	 */
+	UPROPERTY(Config, BlueprintReadWrite, EditAnywhere, Category = "Virtual Camera")
+	EVCamTutorialCompletionState VCamTutorialCompletationState = EVCamTutorialCompletionState::Pending;
+
+	/** @return Whether the VCam tutorial has been completed. */
+	UFUNCTION(BlueprintPure, Category = "Virtual Camera")
+	bool IsTutorialCompleted() const { return VCamTutorialCompletationState == EVCamTutorialCompletionState::Completed; }
+};
\ No newline at end of file
