Allow Blueprints to change viewports into game view

diff --git a/Engine/Plugins/VirtualProduction/VirtualCamera/Source/VirtualCamera/Private/FunctionLibraries/GameViewFunctionLibrary.cpp b/Engine/Plugins/VirtualProduction/VirtualCamera/Source/VirtualCamera/Private/FunctionLibraries/GameViewFunctionLibrary.cpp
new file mode 100644
--- /dev/null
+++ b/Engine/Plugins/VirtualProduction/VirtualCamera/Source/VirtualCamera/Private/FunctionLibraries/GameViewFunctionLibrary.cpp
@@ -0,0 +1,86 @@
+// Copyright Epic Games, Inc. All Rights Reserved.
+
+#include "FunctionLibraries/GameViewFunctionLibrary.h"
+
+#include "EVCamTargetViewportID.h"
+
+#if WITH_EDITOR
+#include "SLevelViewport.h"
+#include "VirtualCamera.h"
+#endif
+
+void UGameViewFunctionLibrary::ToggleGameView(EVCamTargetViewportID ViewportID)
+{
+#if WITH_EDITOR
+	if (const TSharedPtr<SLevelViewport> Viewport = UE::VCamCore::GetLevelViewport(ViewportID))
+	{
+		Viewport->ToggleGameView();
+	}
+	else
+	{
+		UE_LOG(LogVirtualCamera, Warning, TEXT("Failed to ToggleGameView for viewport %s"), *UE::VCamCore::ViewportIdToString(ViewportID));
+	}
+#endif
+}
+
+bool UGameViewFunctionLibrary::CanToggleGameView(EVCamTargetViewportID ViewportID)
+{
+#if WITH_EDITOR
+	const TSharedPtr<SLevelViewport> Viewport = UE::VCamCore::GetLevelViewport(ViewportID);
+	return Viewport && Viewport->CanToggleGameView();
+#else
+	return false;
+#endif
+}
+
+bool UGameViewFunctionLibrary::IsInGameView(EVCamTargetViewportID ViewportID)
+{
+#if WITH_EDITOR
+	const TSharedPtr<SLevelViewport> Viewport = UE::VCamCore::GetLevelViewport(ViewportID);
+	return Viewport && Viewport->IsInGameView();
+#else
+	return true;
+#endif
+}
+
+void UGameViewFunctionLibrary::SetGameViewEnabled(EVCamTargetViewportID ViewportID, bool bIsEnabled)
+{
+	if ((IsInGameView(ViewportID) && !bIsEnabled) || (!IsInGameView(ViewportID) && bIsEnabled))
+	{
+		ToggleGameView(ViewportID);
+	}
+}
+
+void UGameViewFunctionLibrary::SetGameViewEnabledForAllViewports(bool bIsEnabled)
+{
+	for (int32 i = 0; i < static_cast<int32>(EVCamTargetViewportID::Count); ++i)
+	{
+		SetGameViewEnabled(static_cast<EVCamTargetViewportID>(i), bIsEnabled);
+	}
+}
+
+TMap<EVCamTargetViewportID, bool> UGameViewFunctionLibrary::SnapshotGameViewStates()
+{
+	TMap<EVCamTargetViewportID, bool> Result;
+#if WITH_EDITOR
+	for (int32 i = 0; i < static_cast<int32>(EVCamTargetViewportID::Count); ++i)
+	{
+		const EVCamTargetViewportID ViewportID = static_cast<EVCamTargetViewportID>(i);
+		if (const TSharedPtr<SLevelViewport> Viewport = UE::VCamCore::GetLevelViewport(ViewportID))
+		{
+			Result.Add(ViewportID, Viewport->IsInGameView());
+		}
+	}
+#endif
+	return Result;
+}
+
+void UGameViewFunctionLibrary::RestoreGameViewStates(const TMap<EVCamTargetViewportID, bool>& Snapshot)
+{
+#if WITH_EDITOR
+	for (const TPair<EVCamTargetViewportID, bool>& SnapshotPair : Snapshot)
+	{
+		SetGameViewEnabled(SnapshotPair.Key, SnapshotPair.Value);
+	}
+#endif
+}
diff --git a/Engine/Plugins/VirtualProduction/VirtualCamera/Source/VirtualCamera/Public/FunctionLibraries/GameViewFunctionLibrary.h b/Engine/Plugins/VirtualProduction/VirtualCamera/Source/VirtualCamera/Public/FunctionLibraries/GameViewFunctionLibrary.h
new file mode 100644
--- /dev/null
+++ b/Engine/Plugins/VirtualProduction/VirtualCamera/Source/VirtualCamera/Public/FunctionLibraries/GameViewFunctionLibrary.h
@@ -0,0 +1,68 @@
+// Copyright Epic Games, Inc. All Rights Reserved.
+
+#pragma once
+
+#include "CoreMinimal.h"
+#include "Kismet/BlueprintFunctionLibrary.h"
+#include "GameViewFunctionLibrary.generated.h"
+
+enum class EVCamTargetViewportID : uint8;
+
+/** Library for changing the viewport into game view. Game view shows the scene as it appears in game. */
+UCLASS()
+class VIRTUALCAMERA_API UGameViewFunctionLibrary : public UBlueprintFunctionLibrary
+{
+	GENERATED_BODY()
+public:
+
+	/**
+	 * Toggles the game view.
+	 * Game view shows the scene as it appears in game.
+	 * 
+	 * @note Only works in editor builds.
+	 */
+	UFUNCTION(BlueprintCallable, Category = "Virtual Camera | Game View")
+	static void ToggleGameView(EVCamTargetViewportID ViewportID);
+
+	/**
+	 * @return Whether we can toggle game view. 
+	 * @note Only works in editor builds.
+	 */
+	UFUNCTION(BlueprintCallable, Category = "Virtual Camera | Game View")
+	static bool CanToggleGameView(EVCamTargetViewportID ViewportID);
+
+	/**
+	 * @return Whether we are in game view.
+	 * @note Only works in editor builds.
+	 */
+	UFUNCTION(BlueprintCallable, Category = "Virtual Camera | Game View")
+	static bool IsInGameView(EVCamTargetViewportID ViewportID);
+
+	/**
+	 * Sets whether the specified viewport should have the game view enabled.
+	 * @note Only works in editor builds.
+	 * */
+	UFUNCTION(BlueprintCallable, Category = "Virtual Camera | Game View")
+	static void SetGameViewEnabled(EVCamTargetViewportID ViewportID, bool bIsEnabled);
+
+	/**
+	 * Sets the game view mode for all open viewports. 
+	 * @note Only works in editor builds.
+	 */
+	UFUNCTION(BlueprintCallable, Category = "Virtual Camera | Game View")
+	static void SetGameViewEnabledForAllViewports(bool bIsEnabled);
+
+	/**
+	 * @return The state of each viewport's state. If the viewport is not currently open, then the returned map will not have any entry for it. 
+	 * @note Only works in editor builds.
+	 */
+	UFUNCTION(BlueprintPure, Category = "Virtual Camera | Game View")
+	static TMap<EVCamTargetViewportID, bool> SnapshotGameViewStates();
+
+	/**
+	 * Sets the game view of the viewports as specified in the map. Util for using together with SnapshotGameViewStates. 
+	 * @note Only works in editor builds.
+	 */
+	UFUNCTION(BlueprintCallable, Category = "Virtual Camera | Game View")
+	static void RestoreGameViewStates(const TMap<EVCamTargetViewportID, bool>& Snapshot);
+};
