Fix: Make VCam settings getter static and add setter which saves the settings object.
The functions are not yet being used in any Blueprints (but this will be used later still for a 5.4 VCam change by [REDACTED_PERSON]) so this change is safe.

diff --git a/Engine/Plugins/VirtualProduction/VirtualCamera/Source/VirtualCamera/Public/VCamUserSettings.h b/Engine/Plugins/VirtualProduction/VirtualCamera/Source/VirtualCamera/Public/VCamUserSettings.h
--- a/Engine/Plugins/VirtualProduction/VirtualCamera/Source/VirtualCamera/Public/VCamUserSettings.h
+++ b/Engine/Plugins/VirtualProduction/VirtualCamera/Source/VirtualCamera/Public/VCamUserSettings.h
@@ -23,17 +23,28 @@ class VIRTUALCAMERA_API UVirtualCameraUserSettings : public UObject
 public:
 
 	UFUNCTION(BlueprintPure, Category = "Virtual Camera", meta = (Keywords = "Get Settings Virtual Camera VCam"))
-	UVirtualCameraUserSettings* GetSettings() { return GetMutableDefault<UVirtualCameraUserSettings>(); }
+	static UVirtualCameraUserSettings* GetSettings() { return GetMutableDefault<UVirtualCameraUserSettings>(); }
 
 	/**
 	 * Indicates whether the VCam tutorial is completed.
 	 * You can manually reset this to Pending if you want to retake the tutorial.
 	 * The tutorial shown in the default VCamHUD, e.g. to teach gestures.
 	 */
-	UPROPERTY(Config, BlueprintReadWrite, EditAnywhere, Category = "Virtual Camera")
-	EVCamTutorialCompletionState VCamTutorialCompletationState = EVCamTutorialCompletionState::Pending;
+	UPROPERTY(Config, BlueprintReadWrite, EditAnywhere, BlueprintSetter = "SetTutorialCompletionState", Category = "Virtual Camera")
+	EVCamTutorialCompletionState VCamTutorialCompletionState = EVCamTutorialCompletionState::Pending;
 
 	/** @return Whether the VCam tutorial has been completed. */
 	UFUNCTION(BlueprintPure, Category = "Virtual Camera")
-	bool IsTutorialCompleted() const { return VCamTutorialCompletationState == EVCamTutorialCompletionState::Completed; }
+	bool IsTutorialCompleted() const { return VCamTutorialCompletionState == EVCamTutorialCompletionState::Completed; }
+
+	/** Sets the tutorial completion state and saves it to the config file if it has changed. */
+	UFUNCTION(BlueprintCallable, Category = "Virtual Camera")
+	void SetTutorialCompletionState(EVCamTutorialCompletionState NewState)
+	{
+		if (VCamTutorialCompletionState != NewState)
+		{
+			VCamTutorialCompletionState = NewState;
+			SaveConfig();
+		}
+	}
 };
\ No newline at end of file
