VCam post smoothing initial checkin

#virtualized

diff --git a/Engine/Plugins/VirtualProduction/VirtualCamera/Source/VirtualCamera/Private/LevelSequence/VirtualCameraClipsMetaData.cpp b/Engine/Plugins/VirtualProduction/VirtualCamera/Source/VirtualCamera/Private/LevelSequence/VirtualCameraClipsMetaData.cpp
--- a/Engine/Plugins/VirtualProduction/VirtualCamera/Source/VirtualCamera/Private/LevelSequence/VirtualCameraClipsMetaData.cpp
+++ b/Engine/Plugins/VirtualProduction/VirtualCamera/Source/VirtualCamera/Private/LevelSequence/VirtualCameraClipsMetaData.cpp
@@ -16,6 +16,7 @@ const FName UVirtualCameraClipsMetaData::AssetRegistryTag_bIsNoGood = "ClipsMeta
 const FName UVirtualCameraClipsMetaData::AssetRegistryTag_bIsFlagged = "ClipsMetaData_bIsFlagged";
 const FName UVirtualCameraClipsMetaData::AssetRegistryTag_FavoriteLevel = "ClipsMetaData_FavoriteLevel";
 const FName UVirtualCameraClipsMetaData::AssetRegistryTag_bIsCreatedFromVCam = "ClipsMetaData_bIsCreatedFromVCam";
+const FName UVirtualCameraClipsMetaData::AssetRegistryTag_PostSmoothLevel = "ClipMetaData_PostSmoothLevel";
 
 UVirtualCameraClipsMetaData::UVirtualCameraClipsMetaData(const FObjectInitializer& ObjInit)
 	: Super(ObjInit)
@@ -39,7 +40,8 @@ TSet<FName> UVirtualCameraClipsMetaData::GetAllClipsMetaDataTags()
 		GetClipsMetaDataTag_IsNoGood(),
 		GetClipsMetaDataTag_IsFlagged(),
 		GetClipsMetaDataTag_FavoriteLevel(),
-		GetClipsMetaDataTag_IsCreatedFromVCam()
+		GetClipsMetaDataTag_IsCreatedFromVCam(),
+		GetClipsMetaDataTag_PostSmoothLevel()
 	};
 }
 
@@ -62,6 +64,7 @@ void UVirtualCameraClipsMetaData::ExtendAssetRegistryTags(FAssetRegistryTagsCont
 	Context.AddTag(FAssetRegistryTag(AssetRegistryTag_bIsFlagged, FString::FromInt(bIsFlagged), FAssetRegistryTag::ETagType::TT_Numerical, FAssetRegistryTag::TD_None));
 	Context.AddTag(FAssetRegistryTag(AssetRegistryTag_FavoriteLevel, FString::FromInt(FavoriteLevel), FAssetRegistryTag::ETagType::TT_Alphabetical, FAssetRegistryTag::TD_None));
 	Context.AddTag(FAssetRegistryTag(AssetRegistryTag_bIsCreatedFromVCam, FString::FromInt(bIsCreatedFromVCam), FAssetRegistryTag::ETagType::TT_Numerical, FAssetRegistryTag::TD_None));
+	Context.AddTag(FAssetRegistryTag(AssetRegistryTag_PostSmoothLevel, FString::FromInt(PostSmoothLevel), FAssetRegistryTag::ETagType::TT_Alphabetical, FAssetRegistryTag::TD_None));
 }
 
 #if WITH_EDITOR
diff --git a/Engine/Plugins/VirtualProduction/VirtualCamera/Source/VirtualCamera/Public/LevelSequence/VirtualCameraClipsMetaData.h b/Engine/Plugins/VirtualProduction/VirtualCamera/Source/VirtualCamera/Public/LevelSequence/VirtualCameraClipsMetaData.h
--- a/Engine/Plugins/VirtualProduction/VirtualCamera/Source/VirtualCamera/Public/LevelSequence/VirtualCameraClipsMetaData.h
+++ b/Engine/Plugins/VirtualProduction/VirtualCamera/Source/VirtualCamera/Public/LevelSequence/VirtualCameraClipsMetaData.h
@@ -58,6 +58,9 @@ public:
 
 	/** The asset registry tag that contains whether it was created from a VCam */
 	static const FName AssetRegistryTag_bIsCreatedFromVCam;
+
+	/* The asset registry tag that contains the post smooth level*/
+	static const FName AssetRegistryTag_PostSmoothLevel;
 	
 public:
 
@@ -109,6 +112,10 @@ public:
 	UFUNCTION(BlueprintPure, Category = "VirtualCamera|Clips")
 	static FName GetClipsMetaDataTag_IsCreatedFromVCam() { return AssetRegistryTag_bIsCreatedFromVCam; }
 
+	/** The asset registry tag that contains post smooth level */
+	UFUNCTION(BlueprintPure, Category = "VirtualCamera|Clips")
+	static FName GetClipsMetaDataTag_PostSmoothLevel() { return AssetRegistryTag_PostSmoothLevel; }
+
 	/** Gets all asset registry tags */
 	UFUNCTION(BlueprintPure, Category = "VirtualCamera|Clips")
 	static TSet<FName> GetAllClipsMetaDataTags();
@@ -273,4 +280,8 @@ private:
 	/** Whether the sequence was created from a VCam */
 	UPROPERTY(EditAnywhere, Category = "Clips", BlueprintReadWrite, meta = (AllowPrivateAccess = "true"))
 	bool bIsCreatedFromVCam = true;
+
+	/* The asset registry tag that contains the post smooth level*/
+	UPROPERTY(EditAnywhere, Category = "Clips", BlueprintReadWrite, meta = (AllowPrivateAccess = "true"))
+	float PostSmoothLevel = 0.0f;
 };
\ No newline at end of file
