VCam post smooth: Prebake curves before applying gaussian filter. Support non subsequence takes. Add value field and adjust slider layout. PostSmoothLevel metadata to be integer to be the same as GaussianFilter KernelWidth param

#virtualized

diff --git a/Engine/Plugins/VirtualProduction/VirtualCamera/Source/VirtualCamera/Private/LevelSequence/VirtualCameraClipsMetaData.cpp b/Engine/Plugins/VirtualProduction/VirtualCamera/Source/VirtualCamera/Private/LevelSequence/VirtualCameraClipsMetaData.cpp
--- a/Engine/Plugins/VirtualProduction/VirtualCamera/Source/VirtualCamera/Private/LevelSequence/VirtualCameraClipsMetaData.cpp
+++ b/Engine/Plugins/VirtualProduction/VirtualCamera/Source/VirtualCamera/Private/LevelSequence/VirtualCameraClipsMetaData.cpp
@@ -64,7 +64,7 @@ void UVirtualCameraClipsMetaData::ExtendAssetRegistryTags(FAssetRegistryTagsCont
 	Context.AddTag(FAssetRegistryTag(AssetRegistryTag_bIsFlagged, FString::FromInt(bIsFlagged), FAssetRegistryTag::ETagType::TT_Numerical, FAssetRegistryTag::TD_None));
 	Context.AddTag(FAssetRegistryTag(AssetRegistryTag_FavoriteLevel, FString::FromInt(FavoriteLevel), FAssetRegistryTag::ETagType::TT_Alphabetical, FAssetRegistryTag::TD_None));
 	Context.AddTag(FAssetRegistryTag(AssetRegistryTag_bIsCreatedFromVCam, FString::FromInt(bIsCreatedFromVCam), FAssetRegistryTag::ETagType::TT_Numerical, FAssetRegistryTag::TD_None));
-	Context.AddTag(FAssetRegistryTag(AssetRegistryTag_PostSmoothLevel, FString::FromInt(PostSmoothLevel), FAssetRegistryTag::ETagType::TT_Alphabetical, FAssetRegistryTag::TD_None));
+	Context.AddTag(FAssetRegistryTag(AssetRegistryTag_PostSmoothLevel, FString::FromInt(PostSmoothLevel), FAssetRegistryTag::ETagType::TT_Numerical, FAssetRegistryTag::TD_None));
 }
 
 #if WITH_EDITOR
diff --git a/Engine/Plugins/VirtualProduction/VirtualCamera/Source/VirtualCamera/Public/LevelSequence/VirtualCameraClipsMetaData.h b/Engine/Plugins/VirtualProduction/VirtualCamera/Source/VirtualCamera/Public/LevelSequence/VirtualCameraClipsMetaData.h
--- a/Engine/Plugins/VirtualProduction/VirtualCamera/Source/VirtualCamera/Public/LevelSequence/VirtualCameraClipsMetaData.h
+++ b/Engine/Plugins/VirtualProduction/VirtualCamera/Source/VirtualCamera/Public/LevelSequence/VirtualCameraClipsMetaData.h
@@ -283,5 +283,5 @@ private:
 
 	/* The asset registry tag that contains the post smooth level*/
 	UPROPERTY(EditAnywhere, Category = "Clips", BlueprintReadWrite, meta = (AllowPrivateAccess = "true"))
-	float PostSmoothLevel = 0.0f;
+	int32 PostSmoothLevel = 0;
 };
\ No newline at end of file
