VCamTakesBrowser: Shows dirty state. Shows full path on hover

#virtualized

diff --git a/Engine/Plugins/VirtualProduction/VirtualCamera/Source/VirtualCamera/Private/FunctionLibraries/VCamBlueprintFunctionLibrary.cpp b/Engine/Plugins/VirtualProduction/VirtualCamera/Source/VirtualCamera/Private/FunctionLibraries/VCamBlueprintFunctionLibrary.cpp
--- a/Engine/Plugins/VirtualProduction/VirtualCamera/Source/VirtualCamera/Private/FunctionLibraries/VCamBlueprintFunctionLibrary.cpp
+++ b/Engine/Plugins/VirtualProduction/VirtualCamera/Source/VirtualCamera/Private/FunctionLibraries/VCamBlueprintFunctionLibrary.cpp
@@ -713,6 +713,16 @@ void UVCamBlueprintFunctionLibrary::SetPlaybackSpeed(float Value)
 #endif
 }
 
+bool UVCamBlueprintFunctionLibrary::IsAssetDirty(const FAssetData& AssetData)
+{
+	UPackage* Package = AssetData.GetPackage();
+	if (Package && Package->IsDirty())
+	{
+		return true;
+	}
+	return false;
+}
+
 #if WITH_EDITOR
 TWeakPtr<ISequencer> UVCamBlueprintFunctionLibrary::GetSequencer()
 {
diff --git a/Engine/Plugins/VirtualProduction/VirtualCamera/Source/VirtualCamera/Public/FunctionLibraries/VCamBlueprintFunctionLibrary.h b/Engine/Plugins/VirtualProduction/VirtualCamera/Source/VirtualCamera/Public/FunctionLibraries/VCamBlueprintFunctionLibrary.h
--- a/Engine/Plugins/VirtualProduction/VirtualCamera/Source/VirtualCamera/Public/FunctionLibraries/VCamBlueprintFunctionLibrary.h
+++ b/Engine/Plugins/VirtualProduction/VirtualCamera/Source/VirtualCamera/Public/FunctionLibraries/VCamBlueprintFunctionLibrary.h
@@ -243,6 +243,10 @@ public:
 	UFUNCTION(BlueprintPure, Category = "VirtualCamera")
 	static bool MultiTraceHitProxyOnViewport(const FVector2D& InScreenPosition, EVCamTargetViewportID InTargetViewport, FVCamTraceHitProxyQueryParams InQueryParams, TArray<FVCamTraceHitProxyResult>& Result);
 
+	/* Checks if the asset is dirty */
+	UFUNCTION(BlueprintPure, Category = "VirtualCamera")
+	static bool IsAssetDirty(const FAssetData& AssetData);
+
 private:
 
 #if WITH_EDITOR
