VCam: Avoid loading asset when checking dirty in TakesBrowser

diff --git a/Engine/Plugins/VirtualProduction/VirtualCamera/Source/VirtualCamera/Private/FunctionLibraries/VCamBlueprintFunctionLibrary.cpp b/Engine/Plugins/VirtualProduction/VirtualCamera/Source/VirtualCamera/Private/FunctionLibraries/VCamBlueprintFunctionLibrary.cpp
--- a/Engine/Plugins/VirtualProduction/VirtualCamera/Source/VirtualCamera/Private/FunctionLibraries/VCamBlueprintFunctionLibrary.cpp
+++ b/Engine/Plugins/VirtualProduction/VirtualCamera/Source/VirtualCamera/Private/FunctionLibraries/VCamBlueprintFunctionLibrary.cpp
@@ -715,10 +715,13 @@ void UVCamBlueprintFunctionLibrary::SetPlaybackSpeed(float Value)
 
 bool UVCamBlueprintFunctionLibrary::IsAssetDirty(const FAssetData& AssetData)
 {
-	UPackage* Package = AssetData.GetPackage();
-	if (Package && Package->IsDirty())
+	if (AssetData.IsAssetLoaded())
 	{
-		return true;
+		UPackage* Package = AssetData.GetPackage();
+		if (Package && Package->IsDirty())
+		{
+			return true;
+		}
 	}
 	return false;
 }
