Improve editor start-up time by async loading: VirtualCameraEditor now async loads /VirtualCamera/VCamActor.VCamActor_C

diff --git a/Engine/Plugins/VirtualProduction/VirtualCamera/Source/VirtualCameraEditor/Private/VirtualCameraEditorModule.cpp b/Engine/Plugins/VirtualProduction/VirtualCamera/Source/VirtualCameraEditor/Private/VirtualCameraEditorModule.cpp
--- a/Engine/Plugins/VirtualProduction/VirtualCamera/Source/VirtualCameraEditor/Private/VirtualCameraEditorModule.cpp
+++ b/Engine/Plugins/VirtualProduction/VirtualCamera/Source/VirtualCameraEditor/Private/VirtualCameraEditorModule.cpp
@@ -9,6 +9,7 @@
 #include "IVPUtilitiesEditorModule.h"
 #include "LevelEditor.h"
 #include "LevelEditorOutlinerSettings.h"
+#include "UObject/UObjectGlobals.h"
 
 #define LOCTEXT_NAMESPACE "FVirtualCameraEditorModule"
 
@@ -33,6 +34,7 @@ private:
 		if (const FPlacementCategoryInfo* Info = IVPUtilitiesEditorModule::Get().GetVirtualProductionPlacementCategoryInfo()
 			; Info && GEditor)
 		{
+			
 			FAssetData VirtualCamera2ActorAssetData(
 				TEXT("/VirtualCamera/VCamActor"),
 				TEXT("/VirtualCamera"),
@@ -54,9 +56,13 @@ private:
 
 			// ... but if you search for it by text this is needed to make it show up (without having the user load it manually).
 			// The search filters everything in the FBuiltInPlacementCategories::AllClasses category only;
-			// it contains 1. loaded BP classes and 2. specialized actor factories. This manual load adds it to case 1. 
-			const UClass* VCamActorBlueprintClass = LoadClass<UObject>(nullptr, TEXT("/VirtualCamera/VCamActor.VCamActor_C"));
-			UE_CLOG(VCamActorBlueprintClass == nullptr, LogVirtualCameraEditor, Warning, TEXT("Failed to load '/VirtualCamera/VCamActor.VCamActor_C'. Has the Blueprint been moved?"));
+			// it contains 1. loaded BP classes and 2. specialized actor factories. This manual load adds it to case 1.
+			LoadAssetAsync(
+				FTopLevelAssetPath(TEXT("/VirtualCamera/VCamActor"), TEXT("VCamActor_C")),
+				FLoadAssetAsyncDelegate::CreateLambda([](const FTopLevelAssetPath& /*AssetPath*/, UObject* LoadedObject, EAsyncLoadingResult::Type)
+				{
+					UE_CLOG(LoadedObject == nullptr, LogVirtualCameraEditor, Warning, TEXT("Failed to load '/VirtualCamera/VCamActor.VCamActor_C'. Has the Blueprint been moved?"));
+				}));
 		}
 	}
 
