Added VCam photo config settings

diff --git a/Engine/Plugins/VirtualProduction/VirtualCamera/Source/VirtualCamera/Private/VCamUserSettings.cpp b/Engine/Plugins/VirtualProduction/VirtualCamera/Source/VirtualCamera/Private/VCamUserSettings.cpp
new file mode 100644
--- /dev/null
+++ b/Engine/Plugins/VirtualProduction/VirtualCamera/Source/VirtualCamera/Private/VCamUserSettings.cpp
@@ -0,0 +1,11 @@
+// Copyright Epic Games, Inc. All Rights Reserved.
+
+#include "VCamUserSettings.h"
+#include "Misc/Paths.h"
+
+UVirtualCameraUserSettings::UVirtualCameraUserSettings()
+{
+	PhotoSaveLocation.Path = FPaths::Combine(*FPaths::ProjectSavedDir(), TEXT("VCamPhotos"));
+	TextureSaveLocation.Path = TEXT("/Game/VCamPhotos");
+	PhotoSaveMode = EVCamPhotoSaveMode::Both;
+}
\ No newline at end of file
diff --git a/Engine/Plugins/VirtualProduction/VirtualCamera/Source/VirtualCamera/Public/VCamUserSettings.h b/Engine/Plugins/VirtualProduction/VirtualCamera/Source/VirtualCamera/Public/VCamUserSettings.h
--- a/Engine/Plugins/VirtualProduction/VirtualCamera/Source/VirtualCamera/Public/VCamUserSettings.h
+++ b/Engine/Plugins/VirtualProduction/VirtualCamera/Source/VirtualCamera/Public/VCamUserSettings.h
@@ -15,6 +15,15 @@ enum class EVCamTutorialCompletionState : uint8
 	Completed
 };
 
+UENUM(BlueprintType)
+enum class EVCamPhotoSaveMode : uint8
+{
+	Off,
+	PNG,
+	Texture,
+	Both
+};
+
 UCLASS(BlueprintType, Config = EditorPerProjectUserSettings)
 class VIRTUALCAMERA_API UVirtualCameraUserSettings : public UObject
 {
@@ -25,6 +34,8 @@ public:
 	UFUNCTION(BlueprintPure, Category = "Virtual Camera", meta = (Keywords = "Get Settings Virtual Camera VCam"))
 	static UVirtualCameraUserSettings* GetSettings() { return GetMutableDefault<UVirtualCameraUserSettings>(); }
 
+	UVirtualCameraUserSettings();
+
 	/**
 	 * Indicates whether the VCam tutorial is completed.
 	 * You can manually reset this to Pending if you want to retake the tutorial.
@@ -47,4 +58,16 @@ public:
 			SaveConfig();
 		}
 	}
+
+	/** VCam photo save mode*/
+	UPROPERTY(Config, BlueprintReadWrite, EditAnywhere, Category = "VCam Photo")
+	EVCamPhotoSaveMode PhotoSaveMode;
+
+	/** Directory path to store vcam photos*/
+	UPROPERTY(Config, BlueprintReadWrite, EditAnywhere, Category = "VCam Photo", meta = (EditCondition="PhotoSaveMode==EVCamPhotoSaveMode::PNG || PhotoSaveMode==EVCamPhotoSaveMode::Both"))
+	FDirectoryPath PhotoSaveLocation;
+
+	/** Directory path to store vcam photo texture assets*/
+	UPROPERTY(Config, BlueprintReadWrite, EditAnywhere, Category = "VCam Photo", meta = (ContentDir, EditCondition="PhotoSaveMode==EVCamPhotoSaveMode::Texture || PhotoSaveMode==EVCamPhotoSaveMode::Both"))
+	FDirectoryPath TextureSaveLocation;
 };
\ No newline at end of file
