Added CineCameraSceneCapture plugin as dependency for Virtual Camera plugin

diff --git a/Engine/Plugins/VirtualProduction/VirtualCamera/VirtualCamera.uplugin b/Engine/Plugins/VirtualProduction/VirtualCamera/VirtualCamera.uplugin
--- a/Engine/Plugins/VirtualProduction/VirtualCamera/VirtualCamera.uplugin
+++ b/Engine/Plugins/VirtualProduction/VirtualCamera/VirtualCamera.uplugin
@@ -48,6 +48,10 @@
 		{
 			"Name" : "VirtualCameraCore",
 			"Enabled": true
+		},
+		{
+			"Name": "CineCameraSceneCapture",
+			"Enabled": true
 		}
 	]
 }
