VCamTakesBrowser: Fix sorting and filtering

#virtualized

diff --git a/Engine/Plugins/VirtualProduction/VirtualCamera/Source/VirtualCamera/Private/FunctionLibraries/AssetFilteringAndSortingFunctionLibrary.cpp b/Engine/Plugins/VirtualProduction/VirtualCamera/Source/VirtualCamera/Private/FunctionLibraries/AssetFilteringAndSortingFunctionLibrary.cpp
--- a/Engine/Plugins/VirtualProduction/VirtualCamera/Source/VirtualCamera/Private/FunctionLibraries/AssetFilteringAndSortingFunctionLibrary.cpp
+++ b/Engine/Plugins/VirtualProduction/VirtualCamera/Source/VirtualCamera/Private/FunctionLibraries/AssetFilteringAndSortingFunctionLibrary.cpp
@@ -7,6 +7,8 @@
 
 #include "Logging/LogMacros.h"
 
+#include "Misc/ComparisonUtility.h"
+
 DEFINE_LOG_CATEGORY_STATIC(LogAssetFilteringAndSorting, All, All);
 
 namespace UE::VirtualCamera::Private
@@ -79,7 +81,7 @@ void UAssetFilteringAndSortingFunctionLibrary::SortByAssetName(TArray<FAssetData
 	SortAssets(Assets,
 		[](const FAssetData& Left, const FAssetData& Right)
 		{
-			return Left.AssetName.LexicalLess(Right.AssetName);
+			return UE::ComparisonUtility::CompareNaturalOrder(Left.AssetName.ToString(), Right.AssetName.ToString()) < 0;
 		}, SortOrder);
 }
 
