VCam: Added BlueprintFileUtils as dependecy to Virtual Camera plugin to prevent compilation errors for VCamPhoto Modifier

diff --git a/Engine/Plugins/VirtualProduction/VirtualCamera/VirtualCamera.uplugin b/Engine/Plugins/VirtualProduction/VirtualCamera/VirtualCamera.uplugin
--- a/Engine/Plugins/VirtualProduction/VirtualCamera/VirtualCamera.uplugin
+++ b/Engine/Plugins/VirtualProduction/VirtualCamera/VirtualCamera.uplugin
@@ -52,6 +52,10 @@
 		{
 			"Name": "CineCameraSceneCapture",
 			"Enabled": true
+		},
+		{
+			"Name": "BlueprintFileUtils",
+			"Enabled": true
 		}
 	]
 }
