Deprecated UAssetFilteringAndSortingFunctionLibrary

diff --git a/Engine/Plugins/VirtualProduction/VirtualCamera/Source/VirtualCamera/Private/FunctionLibraries/AssetFilteringAndSortingFunctionLibrary.cpp b/Engine/Plugins/VirtualProduction/VirtualCamera/Source/VirtualCamera/Private/FunctionLibraries/AssetFilteringAndSortingFunctionLibrary.cpp
--- a/Engine/Plugins/VirtualProduction/VirtualCamera/Source/VirtualCamera/Private/FunctionLibraries/AssetFilteringAndSortingFunctionLibrary.cpp
+++ b/Engine/Plugins/VirtualProduction/VirtualCamera/Source/VirtualCamera/Private/FunctionLibraries/AssetFilteringAndSortingFunctionLibrary.cpp
@@ -13,6 +13,7 @@ DEFINE_LOG_CATEGORY_STATIC(LogAssetFilteringAndSorting, All, All);
 
 namespace UE::VirtualCamera::Private
 {
+	PRAGMA_DISABLE_DEPRECATION_WARNINGS
 	template<typename TType>
 	static bool Sort(TArray<FAssetData>& Assets, FName MetaDataTag, TFunctionRef<bool(const FString& TagType, TType& Converted)> Converter, ESortOrder SortOrder)
 	{
@@ -34,15 +35,16 @@ namespace UE::VirtualCamera::Private
 			}
 		}
 
-		UAssetFilteringAndSortingFunctionLibrary::SortAssets(Assets, [&MetaData](const FAssetData& Left, const FAssetData& Right)
+		UDEPRECATED_AssetFilteringAndSortingFunctionLibrary::SortAssets(Assets, [&MetaData](const FAssetData& Left, const FAssetData& Right)
 		{
 			return MetaData[Left.GetSoftObjectPath()] <= MetaData[Right.GetSoftObjectPath()]; 
 		}, SortOrder);
 		return true;
 	}
+	PRAGMA_ENABLE_DEPRECATION_WARNINGS
 }
 
-TArray<FAssetData> UAssetFilteringAndSortingFunctionLibrary::GetAllAssetsByMetaDataTags(const TSet<FName>& RequiredTags, const TSet<UClass*>& AllowedClasses)
+TArray<FAssetData> UDEPRECATED_AssetFilteringAndSortingFunctionLibrary::GetAllAssetsByMetaDataTags(const TSet<FName>& RequiredTags, const TSet<UClass*>& AllowedClasses)
 {
 	FARFilter Filter;
 	
@@ -63,7 +65,8 @@ TArray<FAssetData> UAssetFilteringAndSortingFunctionLibrary::GetAllAssetsByMetaD
 	return Result;
 }
 
-void UAssetFilteringAndSortingFunctionLibrary::SortByCustomPredicate(TArray<FAssetData>& Assets, FAssetSortingPredicate SortingPredicate, ESortOrder SortOrder)
+PRAGMA_DISABLE_DEPRECATION_WARNINGS
+void UDEPRECATED_AssetFilteringAndSortingFunctionLibrary::SortByCustomPredicate(TArray<FAssetData>& Assets, FAssetSortingPredicate SortingPredicate, ESortOrder SortOrder)
 {
 	if (SortingPredicate.IsBound())
 	{
@@ -75,8 +78,11 @@ void UAssetFilteringAndSortingFunctionLibrary::SortByCustomPredicate(TArray<FAss
 			}, SortOrder);
 	}
 }
+PRAGMA_ENABLE_DEPRECATION_WARNINGS
 
-void UAssetFilteringAndSortingFunctionLibrary::SortByAssetName(TArray<FAssetData>& Assets, ESortOrder SortOrder)
+
+PRAGMA_DISABLE_DEPRECATION_WARNINGS
+void UDEPRECATED_AssetFilteringAndSortingFunctionLibrary::SortByAssetName(TArray<FAssetData>& Assets, ESortOrder SortOrder)
 {
 	SortAssets(Assets,
 		[](const FAssetData& Left, const FAssetData& Right)
@@ -84,8 +90,10 @@ void UAssetFilteringAndSortingFunctionLibrary::SortByAssetName(TArray<FAssetData
 			return UE::ComparisonUtility::CompareNaturalOrder(Left.AssetName.ToString(), Right.AssetName.ToString()) < 0;
 		}, SortOrder);
 }
+PRAGMA_ENABLE_DEPRECATION_WARNINGS
 
-bool UAssetFilteringAndSortingFunctionLibrary::SortByMetaData(TArray<FAssetData>& Assets, FName MetaDataTag, EAssetTagMetaDataSortType MetaDataType, ESortOrder SortOrder)
+PRAGMA_DISABLE_DEPRECATION_WARNINGS
+bool UDEPRECATED_AssetFilteringAndSortingFunctionLibrary::SortByMetaData(TArray<FAssetData>& Assets, FName MetaDataTag, EAssetTagMetaDataSortType MetaDataType, ESortOrder SortOrder)
 {
 	switch (MetaDataType)
 	{
@@ -111,8 +119,10 @@ bool UAssetFilteringAndSortingFunctionLibrary::SortByMetaData(TArray<FAssetData>
 		return false;
 	}
 }
+PRAGMA_ENABLE_DEPRECATION_WARNINGS
 
-void UAssetFilteringAndSortingFunctionLibrary::SortAssets(TArray<FAssetData>& Assets, TFunctionRef<bool(const FAssetData& Left, const FAssetData& Right)> Predicate, ESortOrder SortOrder)
+PRAGMA_DISABLE_DEPRECATION_WARNINGS
+void UDEPRECATED_AssetFilteringAndSortingFunctionLibrary::SortAssets(TArray<FAssetData>& Assets, TFunctionRef<bool(const FAssetData& Left, const FAssetData& Right)> Predicate, ESortOrder SortOrder)
 {
 	// Careful this would be undefined behaviour: TFunctionRef Variable = [](){}; 
 	auto ReversePredicate = [&Predicate](const FAssetData& Left, const FAssetData& Right) { return !Predicate(Left, Right); };
@@ -121,3 +131,4 @@ void UAssetFilteringAndSortingFunctionLibrary::SortAssets(TArray<FAssetData>& As
 	
 	Assets.Sort(PredicateToUse);
 }
+PRAGMA_ENABLE_DEPRECATION_WARNINGS
diff --git a/Engine/Plugins/VirtualProduction/VirtualCamera/Source/VirtualCamera/Public/FunctionLibraries/AssetFilteringAndSortingFunctionLibrary.h b/Engine/Plugins/VirtualProduction/VirtualCamera/Source/VirtualCamera/Public/FunctionLibraries/AssetFilteringAndSortingFunctionLibrary.h
--- a/Engine/Plugins/VirtualProduction/VirtualCamera/Source/VirtualCamera/Public/FunctionLibraries/AssetFilteringAndSortingFunctionLibrary.h
+++ b/Engine/Plugins/VirtualProduction/VirtualCamera/Source/VirtualCamera/Public/FunctionLibraries/AssetFilteringAndSortingFunctionLibrary.h
@@ -13,42 +13,50 @@
 struct FAssetData;
 
 UENUM(BlueprintType)
-enum class ESortOrder : uint8
+enum class UE_DEPRECATED(5.5, "Use EEditorAssetSortOrder instead") ESortOrder : uint8
 {
 	Ascending,
 	Descending
 };
 
 UENUM(BlueprintType)
-enum class EAssetTagMetaDataSortType : uint8
+enum class UE_DEPRECATED(5.5, "Use EEditorAssetMetaDataSortType instead") EAssetTagMetaDataSortType : uint8
 {
 	String,
 	Numeric,
 	DateTime
 };
 
+class UE_DEPRECATED(5.5, "Use the equivalent delegates for UEditorAssetSubsystem and UAssetRegistryHelpers instead") FAssetSortingPredicate;
 DECLARE_DYNAMIC_DELEGATE_RetVal_TwoParams(bool, FAssetSortingPredicate, const FAssetData&, Left, const FAssetData&, Right);
 
 /** This library's purpose is to facilitate Blueprints to discover assets using some filters and sort them. */
-UCLASS(BlueprintType)
-class VIRTUALCAMERA_API UAssetFilteringAndSortingFunctionLibrary : public UBlueprintFunctionLibrary
+UCLASS(BlueprintType, Deprecated)
+class UE_DEPRECATED(5.5, "Use the equivalent function in UEditorAssetSubsystem and UAssetRegistryHelpers instead") VIRTUALCAMERA_API UDEPRECATED_AssetFilteringAndSortingFunctionLibrary : public UBlueprintFunctionLibrary
 {
 	GENERATED_BODY()
 public:
 
 	/** Gets all assets which have the given tags. */
-	UFUNCTION(BlueprintCallable, Category = "VirtualCamera | Clips")
+	UFUNCTION(BlueprintCallable, Category = "VirtualCamera | Clips", meta = (DeprecatedFunction, DeprecationMessage = "Use EditorAssetSubsystem's GetAllAssetsByMetaDataTags instead."))
+	UE_DEPRECATED(5.5, "Use UEditorAssetSubsystem::GetAllAssetsByMetaDataTags instead.")
 	static TArray<FAssetData> GetAllAssetsByMetaDataTags(const TSet<FName>& RequiredTags, const TSet<UClass*>& AllowedClasses);
 
 	/** Sorts the assets based on a custom Blueprint delegate.
 	 * @param SortingPredicate Implements a Left <= Right relation
 	 */
-	UFUNCTION(BlueprintCallable, Category = "VirtualCamera | Clips")
+	PRAGMA_DISABLE_DEPRECATION_WARNINGS
+	UFUNCTION(BlueprintCallable, Category = "VirtualCamera | Clips", meta = (DeprecatedFunction, DeprecationMessage = "Use AssetRegistryHelpers's SortByCustomPredicate instead."))
+	UE_DEPRECATED(5.5, "Use UAssetRegistryHelpers::SortByCustomPredicate instead.")
 	static void SortByCustomPredicate(UPARAM(Ref) TArray<FAssetData>& Assets, FAssetSortingPredicate SortingPredicate, ESortOrder SortOrder = ESortOrder::Ascending);
+	PRAGMA_ENABLE_DEPRECATION_WARNINGS
 
 	/** Sorts the assets by their asset name. */
-	UFUNCTION(BlueprintCallable, Category = "VirtualCamera | Clips")
+	PRAGMA_DISABLE_DEPRECATION_WARNINGS
+	UFUNCTION(BlueprintCallable, Category = "VirtualCamera | Clips", meta = (DeprecatedFunction, DeprecationMessage = "Use AssetRegistryHelpers's SortByAssetName instead."))
+	UE_DEPRECATED(5.5, "Use UAssetRegistryHelpers::SortByAssetName instead.")
 	static void SortByAssetName(UPARAM(Ref) TArray<FAssetData>& Assets, ESortOrder SortOrder = ESortOrder::Ascending);
+	PRAGMA_ENABLE_DEPRECATION_WARNINGS
 	
 	/**
 	 * Sorts the assets based on their meta data's type.
@@ -59,12 +67,18 @@ public:
 	 * @param SortOrder Whether to sort ascending or descending
 	 * @return Whether it was possible to compare all the meta data successfully.
 	 */
-	UFUNCTION(BlueprintCallable, Category = "VirtualCamera | Clips")
+	PRAGMA_DISABLE_DEPRECATION_WARNINGS
+	UFUNCTION(BlueprintCallable, Category = "VirtualCamera | Clips", meta = (DeprecatedFunction, DeprecationMessage = "Use EditorAssetSubsystem's SortByMetaData instead."))
+	UE_DEPRECATED(5.5, "Use UEditorAssetSubsystem::SortByMetaData instead.")
 	static bool SortByMetaData(UPARAM(Ref) TArray<FAssetData>& Assets, FName MetaDataTag, EAssetTagMetaDataSortType MetaDataType, ESortOrder SortOrder = ESortOrder::Ascending);
+	PRAGMA_ENABLE_DEPRECATION_WARNINGS
 	
+	PRAGMA_DISABLE_DEPRECATION_WARNINGS
 	/**
 	 * Util that does the actual sorting
 	 * @param Predicate Implements a Left <= Right relation
 	 */
+	UE_DEPRECATED(5.5, "This function will not be replaced. Write your own version.")
 	static void SortAssets(TArray<FAssetData>& Assets, TFunctionRef<bool(const FAssetData& Left, const FAssetData& Right)> Predicate, ESortOrder SortOrder);
+	PRAGMA_ENABLE_DEPRECATION_WARNINGS
 };
