VCam: Updated AutoFocus trace to use MultiTrace Hit Proxy first to detect Actors with NoCollision and use Object Trace as fall back

#virtualized

diff --git a/Engine/Plugins/VirtualProduction/VirtualCamera/Source/VirtualCamera/Private/FunctionLibraries/VCamBlueprintFunctionLibrary.cpp b/Engine/Plugins/VirtualProduction/VirtualCamera/Source/VirtualCamera/Private/FunctionLibraries/VCamBlueprintFunctionLibrary.cpp
--- a/Engine/Plugins/VirtualProduction/VirtualCamera/Source/VirtualCamera/Private/FunctionLibraries/VCamBlueprintFunctionLibrary.cpp
+++ b/Engine/Plugins/VirtualProduction/VirtualCamera/Source/VirtualCamera/Private/FunctionLibraries/VCamBlueprintFunctionLibrary.cpp
@@ -575,6 +575,7 @@ namespace UE::VirtualCamera::Private
 		}
 		
 		const FIntRect QueryRect(MinX, MinY, MaxX + 1, MaxY + 1);
+		Viewport.InvalidateHitProxy();
 		Viewport.GetHitProxyMap(QueryRect, Result);
 		return { TestSizeX, TestSizeY, Result };
 	}
