Sequencer: Generalize UVCamClipsMetaData to Sequencer. This meta data shared by various cinematics tools.

#virtualized

diff --git a/Engine/Plugins/VirtualProduction/VirtualCamera/Source/VirtualCamera/Private/LevelSequence/VCamTakesMetaDataMigration.cpp b/Engine/Plugins/VirtualProduction/VirtualCamera/Source/VirtualCamera/Private/LevelSequence/VCamTakesMetaDataMigration.cpp
new file mode 100644
--- /dev/null
+++ b/Engine/Plugins/VirtualProduction/VirtualCamera/Source/VirtualCamera/Private/LevelSequence/VCamTakesMetaDataMigration.cpp
@@ -0,0 +1,261 @@
+// Copyright Epic Games, Inc. All Rights Reserved.
+
+#include "LevelSequence/VCamTakesMetaDataMigration.h"
+
+#include "AssetRegistry/AssetData.h"
+#include "AssetRegistry/IAssetRegistry.h"
+#include "HAL/ConsoleManager.h"
+#include "LevelSequence.h"
+#include "LevelSequence/VirtualCameraClipsMetaData.h"
+#include "LevelSequenceShotMetaDataLibrary.h"
+#include "Misc/ScopedSlowTask.h"
+#include "UObject/Package.h"
+#include "UObject/SavePackage.h"
+
+namespace UE::VirtualCamera::MigrationDetail
+{
+static TAutoConsoleVariable<bool> CVarAutoMigrateLevelSequenceOnAccess(
+	TEXT("VirtualCamera.AutoMigrateLevelSequenceOnAccess"),
+	false,
+	TEXT("Whenever VirtualCamera opens a level sequence, whether to automatically migrate the meta data from the old UVirtualCameraClipsMetaData to the new UMovieSceneShotMetaData.")
+);
+
+static FAutoConsoleCommand CMigrateLevelSequenceCommand(
+	TEXT("VirtualCamera.MigrateLevelSequences"),
+	TEXT("Goes through all ULevelSequences in the project and migrates the UVirtualCameraClipsMetaData to UMovieSceneShotMetaData.\nWARNING: This may take a while."),
+	FConsoleCommandDelegate::CreateLambda([]()
+	{
+	   UVCamTakesMetaDataMigration::SLOW_MigrateAllVCamTakesMetaDataInProject();
+	})
+);
+}
+
+PRAGMA_DISABLE_DEPRECATION_WARNINGS
+bool UVCamTakesMetaDataMigration::GetIsNoGood(ULevelSequence* InLevelSequence, bool& bOutNoGood)
+{
+	if (ULevelSequenceShotMetaDataLibrary::GetIsNoGood(InLevelSequence, bOutNoGood))
+	{
+		return true;
+	}
+
+#if WITH_EDITOR
+	const UVirtualCameraClipsMetaData* MetaData = InLevelSequence ? InLevelSequence->FindMetaData<UVirtualCameraClipsMetaData>() : nullptr;
+	if (MetaData)
+	{
+		bOutNoGood = MetaData->GetIsNoGood();
+		return true;
+	}
+#endif
+	
+	return false;
+}
+PRAGMA_ENABLE_DEPRECATION_WARNINGS
+
+PRAGMA_DISABLE_DEPRECATION_WARNINGS
+bool UVCamTakesMetaDataMigration::GetIsFlagged(ULevelSequence* InLevelSequence, bool& bOutIsFlagged)
+{
+	if (ULevelSequenceShotMetaDataLibrary::GetIsFlagged(InLevelSequence, bOutIsFlagged))
+	{
+		return true;
+	}
+
+#if WITH_EDITOR
+	const UVirtualCameraClipsMetaData* MetaData = InLevelSequence ? InLevelSequence->FindMetaData<UVirtualCameraClipsMetaData>() : nullptr;
+	if (MetaData)
+	{
+		bOutIsFlagged = MetaData->GetIsFlagged();
+		return true;
+	}
+#endif
+	
+	return false;
+}
+PRAGMA_ENABLE_DEPRECATION_WARNINGS
+
+PRAGMA_DISABLE_DEPRECATION_WARNINGS
+bool UVCamTakesMetaDataMigration::GetFavoriteLevel(ULevelSequence* InLevelSequence, int32& OutFavoriteLevel)
+{
+	if (ULevelSequenceShotMetaDataLibrary::GetFavoriteRating(InLevelSequence, OutFavoriteLevel))
+	{
+		return true;
+	}
+	
+#if WITH_EDITOR
+	const UVirtualCameraClipsMetaData* MetaData = InLevelSequence ? InLevelSequence->FindMetaData<UVirtualCameraClipsMetaData>() : nullptr;
+	if (MetaData)
+	{
+		OutFavoriteLevel = MetaData->GetFavoriteLevel();
+		return true;
+	}
+#endif
+	
+	return false;
+}
+PRAGMA_ENABLE_DEPRECATION_WARNINGS
+
+PRAGMA_DISABLE_DEPRECATION_WARNINGS
+bool UVCamTakesMetaDataMigration::GetIsNoGoodByAssetData(const FAssetData& InAssetData, bool& bOutNoGood)
+{
+	if (ULevelSequenceShotMetaDataLibrary::GetIsNoGoodByAssetData(InAssetData, bOutNoGood))
+	{
+		return true;
+	}
+
+	return InAssetData.GetTagValue(UVirtualCameraClipsMetaData::AssetRegistryTag_bIsNoGood, bOutNoGood);
+}
+PRAGMA_ENABLE_DEPRECATION_WARNINGS
+
+PRAGMA_DISABLE_DEPRECATION_WARNINGS
+bool UVCamTakesMetaDataMigration::GetIsFlaggedByAssetData(const FAssetData& InAssetData, bool& bOutIsFlagged)
+{
+	if (ULevelSequenceShotMetaDataLibrary::GetIsFlaggedByAssetData(InAssetData, bOutIsFlagged))
+	{
+		return true;
+	}
+
+	return InAssetData.GetTagValue(UVirtualCameraClipsMetaData::AssetRegistryTag_bIsFlagged, bOutIsFlagged);
+}
+PRAGMA_ENABLE_DEPRECATION_WARNINGS
+
+PRAGMA_DISABLE_DEPRECATION_WARNINGS
+bool UVCamTakesMetaDataMigration::GetFavoriteLevelByAssetData(const FAssetData& InAssetData, int32& OutFavoriteLevel)
+{
+	if (ULevelSequenceShotMetaDataLibrary::GetFavoriteRatingByAssetData(InAssetData, OutFavoriteLevel))
+	{
+		return true;
+	}
+
+	return InAssetData.GetTagValue(UVirtualCameraClipsMetaData::AssetRegistryTag_FavoriteLevel, OutFavoriteLevel);
+}
+PRAGMA_ENABLE_DEPRECATION_WARNINGS
+
+void UVCamTakesMetaDataMigration::SetIsNoGood(ULevelSequence* InLevelSequence, bool bInIsNoGood)
+{
+	ULevelSequenceShotMetaDataLibrary::SetIsNoGood(InLevelSequence, bInIsNoGood);
+}
+
+void UVCamTakesMetaDataMigration::SetIsFlagged(ULevelSequence* InLevelSequence, bool bInIsFlagged)
+{
+	ULevelSequenceShotMetaDataLibrary::SetIsFlagged(InLevelSequence, bInIsFlagged);
+}
+
+void UVCamTakesMetaDataMigration::SetFavoriteLevel(ULevelSequence* InLevelSequence, int32 InFavoriteLevel)
+{
+	ULevelSequenceShotMetaDataLibrary::SetFavoriteRating(InLevelSequence, InFavoriteLevel);
+}
+
+bool UVCamTakesMetaDataMigration::NeedsToMigrateVCamMetaData(ULevelSequence* InLevelSequence)
+{
+#if WITH_EDITOR
+	// UMovieSceneShotMetaData contains TOptional values.
+	// UVirtualCameraClipsMetaData always contains values (not wrapped with TOptional).
+	// MigrateVCamTakesMetaData always sets a value during migration. So if all tags have a set value, no migration needs to occur.
+	return InLevelSequence
+		&& InLevelSequence->FindMetaData<UVirtualCameraClipsMetaData>() != nullptr
+		&& (!ULevelSequenceShotMetaDataLibrary::HasIsNoGood(InLevelSequence)
+			|| !ULevelSequenceShotMetaDataLibrary::HasIsFlagged(InLevelSequence)
+			|| !ULevelSequenceShotMetaDataLibrary::HasFavoriteRating(InLevelSequence));
+#else
+	return false;
+#endif
+}
+
+PRAGMA_DISABLE_DEPRECATION_WARNINGS
+bool UVCamTakesMetaDataMigration::NeedsToMigrateVCamMetaDataByAssetData(const FAssetData& InAssetData)
+{
+#if WITH_EDITOR
+	// UMovieSceneShotMetaData contains TOptional values.
+	// UVirtualCameraClipsMetaData always contains values (not wrapped with TOptional).
+	// MigrateVCamTakesMetaData always sets a value during migration. So if all tags have a set value, no migration needs to occur.
+	const bool bHasOldTags = InAssetData.FindTag(UVirtualCameraClipsMetaData::AssetRegistryTag_bIsFlagged)
+		|| InAssetData.FindTag(UVirtualCameraClipsMetaData::AssetRegistryTag_bIsNoGood)
+		|| InAssetData.FindTag(UVirtualCameraClipsMetaData::AssetRegistryTag_FavoriteLevel);
+	const bool bHasMissingNewTags = !ULevelSequenceShotMetaDataLibrary::HasIsNoGoodByAssetData(InAssetData)
+		|| !ULevelSequenceShotMetaDataLibrary::HasIsFlaggedByAssetData(InAssetData)
+		|| !ULevelSequenceShotMetaDataLibrary::HasFavoriteRatingByAssetData(InAssetData);
+	return bHasOldTags && bHasMissingNewTags;
+#else
+	return false;
+#endif
+}
+PRAGMA_ENABLE_DEPRECATION_WARNINGS
+
+PRAGMA_DISABLE_DEPRECATION_WARNINGS
+void UVCamTakesMetaDataMigration::MigrateVCamTakesMetaData(ULevelSequence* InLevelSequence)
+{
+#if WITH_EDITOR
+	const UVirtualCameraClipsMetaData* OldMetaData = InLevelSequence ? InLevelSequence->FindMetaData<UVirtualCameraClipsMetaData>() : nullptr;
+	const bool bNeedsMigration = OldMetaData != nullptr;
+	if (!bNeedsMigration)
+	{
+		return;
+	}
+
+	// Don't overwrite values that may already have been written, e.g. by calling UVCamTakesMetaDataMigration::SetIsNoGood, etc.!
+	if (!ULevelSequenceShotMetaDataLibrary::HasIsNoGood(InLevelSequence))
+	{
+		ULevelSequenceShotMetaDataLibrary::SetIsNoGood(InLevelSequence, OldMetaData->GetIsNoGood());
+	}
+	if (!ULevelSequenceShotMetaDataLibrary::HasIsFlagged(InLevelSequence))
+	{
+		ULevelSequenceShotMetaDataLibrary::SetIsFlagged(InLevelSequence, OldMetaData->GetIsFlagged());
+	}
+	if (!ULevelSequenceShotMetaDataLibrary::HasFavoriteRating(InLevelSequence))
+	{
+		ULevelSequenceShotMetaDataLibrary::SetFavoriteRating(InLevelSequence, OldMetaData->GetFavoriteLevel());
+	}
+
+	// Do not remove the meta data - we'll keep it around for now as backup.
+	// UVirtualCameraClipsMetaData is set up to stop adding the asset tags if UMovieSceneMetaData is present.
+	// InLevelSequence->RemoveMetaData<UVirtualCameraClipsMetaData>();
+#endif
+}
+PRAGMA_ENABLE_DEPRECATION_WARNINGS
+
+bool UVCamTakesMetaDataMigration::GetAutoMigrateAccessedLevelSequencesCVar()
+{
+	return UE::VirtualCamera::MigrationDetail::CVarAutoMigrateLevelSequenceOnAccess.GetValueOnAnyThread();
+}
+
+void UVCamTakesMetaDataMigration::SLOW_MigrateAllVCamTakesMetaDataInProject()
+{
+	FScopedSlowTask SlowTask(10, NSLOCTEXT("VirtualCamera", "Migrate", "Migrating meta data for Level Sequences")); 
+       
+	TArray<FAssetData> LevelSequences;
+	{
+		FScopedSlowTask FilterSlowTask(1, NSLOCTEXT("VirtualCamera", "Migrate.Filtering", "Finding assets to migrate"));
+		FARCompiledFilter Filter;
+		Filter.ClassPaths.Add(FTopLevelAssetPath(ULevelSequence::StaticClass()));
+		IAssetRegistry::Get()->GetAssets(Filter, LevelSequences);
+	}
+
+	FScopedSlowTask MigrationTask(LevelSequences.Num(), NSLOCTEXT("VirtualCamera", "Migrate.DoMigration", "Migrating meta data for Level Sequences")); 
+	for (const FAssetData& AssetData : LevelSequences)
+	{
+		MigrationTask.EnterProgressFrame(1,
+		   FText::Format(
+			  NSLOCTEXT("VirtualCamera", "Migrate.MigrateAssetFmt", "Migrating {0}"),
+			  FText::FromString(AssetData.GetSoftObjectPath().ToString())
+		   )
+		);
+
+		if (!NeedsToMigrateVCamMetaDataByAssetData(AssetData))
+		{
+			continue;
+		}
+          
+		ULevelSequence* Sequence = Cast<ULevelSequence>(AssetData.GetAsset());
+		if (!Sequence)
+		{
+			continue;
+		}
+
+		Sequence->Modify(true);
+		MigrateVCamTakesMetaData(Sequence);
+       
+		UPackage* Package = AssetData.GetPackage();
+		Package->MarkPackageDirty();
+		const FString FilePath = Package->GetLoadedPath().GetLocalFullPath();
+		Package->Save(Package, Sequence, *FilePath, FSavePackageArgs());
+	}
+}
\ No newline at end of file
diff --git a/Engine/Plugins/VirtualProduction/VirtualCamera/Source/VirtualCamera/Private/LevelSequence/VirtualCameraClipsMetaData.cpp b/Engine/Plugins/VirtualProduction/VirtualCamera/Source/VirtualCamera/Private/LevelSequence/VirtualCameraClipsMetaData.cpp
--- a/Engine/Plugins/VirtualProduction/VirtualCamera/Source/VirtualCamera/Private/LevelSequence/VirtualCameraClipsMetaData.cpp
+++ b/Engine/Plugins/VirtualProduction/VirtualCamera/Source/VirtualCamera/Private/LevelSequence/VirtualCameraClipsMetaData.cpp
@@ -2,6 +2,8 @@
 
 #include "LevelSequence/VirtualCameraClipsMetaData.h"
 
+#include "LevelSequence.h"
+#include "LevelSequenceShotMetaDataLibrary.h"
 #include "UObject/AssetRegistryTagsContext.h"
 
 const FName UVirtualCameraClipsMetaData::AssetRegistryTag_FocalLength = "ClipsMetaData_FocalLength";
@@ -21,172 +23,323 @@ const FName UVirtualCameraClipsMetaData::AssetRegistryTag_PostSmoothLevel = "Cli
 UVirtualCameraClipsMetaData::UVirtualCameraClipsMetaData(const FObjectInitializer& ObjInit)
 	: Super(ObjInit)
 {
-	FocalLength = 0;
-	bIsSelected = false;
-	bIsACineCameraRecording = false; 
+}
+
+PRAGMA_DISABLE_DEPRECATION_WARNINGS
+FName UVirtualCameraClipsMetaData::GetClipsMetaDataTag_FocalLength()
+{
+	return AssetRegistryTag_FocalLength;
+}
+
+PRAGMA_DISABLE_DEPRECATION_WARNINGS
+FName UVirtualCameraClipsMetaData::GetClipsMetaDataTag_IsSelected()
+{
+	return AssetRegistryTag_bIsSelected;
+}
+PRAGMA_ENABLE_DEPRECATION_WARNINGS
+
+PRAGMA_DISABLE_DEPRECATION_WARNINGS
+FName UVirtualCameraClipsMetaData::GetClipsMetaDataTag_RecordedLevel()
+{
+	return AssetRegistryTag_RecordedLevelName;
+}
+PRAGMA_ENABLE_DEPRECATION_WARNINGS
+
+PRAGMA_DISABLE_DEPRECATION_WARNINGS
+FName UVirtualCameraClipsMetaData::GetClipsMetaDataTag_FrameCountStart()
+{
+	return AssetRegistryTag_FrameCountStart;
+}
+PRAGMA_ENABLE_DEPRECATION_WARNINGS
+
+PRAGMA_DISABLE_DEPRECATION_WARNINGS
+FName UVirtualCameraClipsMetaData::GetClipsMetaDataTag_FrameCountEnd()
+{
+	return AssetRegistryTag_FrameCountEnd;
+}
+PRAGMA_ENABLE_DEPRECATION_WARNINGS
+
+PRAGMA_DISABLE_DEPRECATION_WARNINGS
+FName UVirtualCameraClipsMetaData::GetClipsMetaDataTag_LengthInFrames()
+{
+	return AssetRegistryTag_LengthInFrames;
+}
+PRAGMA_ENABLE_DEPRECATION_WARNINGS
+
+PRAGMA_DISABLE_DEPRECATION_WARNINGS
+FName UVirtualCameraClipsMetaData::GetClipsMetaDataTag_DisplayRate()
+{
+	return AssetRegistryTag_DisplayRate;
+}
+PRAGMA_ENABLE_DEPRECATION_WARNINGS
+
+PRAGMA_DISABLE_DEPRECATION_WARNINGS
+FName UVirtualCameraClipsMetaData::GetClipsMetaDataTag_IsCineACineCameraRecording()
+{
+	return AssetRegistryTag_bIsACineCameraRecording;
+}
+PRAGMA_ENABLE_DEPRECATION_WARNINGS
+
+PRAGMA_DISABLE_DEPRECATION_WARNINGS
+FName UVirtualCameraClipsMetaData::GetClipsMetaDataTag_IsNoGood()
+{
+	return AssetRegistryTag_bIsNoGood;
+}
+PRAGMA_ENABLE_DEPRECATION_WARNINGS
+
+PRAGMA_DISABLE_DEPRECATION_WARNINGS
+FName UVirtualCameraClipsMetaData::GetClipsMetaDataTag_IsFlagged()
+{
+	return AssetRegistryTag_bIsFlagged;
+}
+PRAGMA_ENABLE_DEPRECATION_WARNINGS
+
+PRAGMA_DISABLE_DEPRECATION_WARNINGS
+FName UVirtualCameraClipsMetaData::GetClipsMetaDataTag_FavoriteLevel()
+{
+	return AssetRegistryTag_FavoriteLevel;
+}
+PRAGMA_ENABLE_DEPRECATION_WARNINGS
+
+PRAGMA_DISABLE_DEPRECATION_WARNINGS
+FName UVirtualCameraClipsMetaData::GetClipsMetaDataTag_IsCreatedFromVCam()
+{
+	return AssetRegistryTag_bIsCreatedFromVCam;
+}
+PRAGMA_ENABLE_DEPRECATION_WARNINGS
+
+FName UVirtualCameraClipsMetaData::GetClipsMetaDataTag_PostSmoothLevel()
+{
+	return AssetRegistryTag_PostSmoothLevel;
 }
 
 TSet<FName> UVirtualCameraClipsMetaData::GetAllClipsMetaDataTags()
 {
 	return {
-		GetClipsMetaDataTag_FocalLength(),
-		GetClipsMetaDataTag_IsSelected(),
-		GetClipsMetaDataTag_RecordedLevel(),
-		GetClipsMetaDataTag_FrameCountStart(),
-		GetClipsMetaDataTag_FrameCountEnd(),
-		GetClipsMetaDataTag_LengthInFrames(),
-		GetClipsMetaDataTag_DisplayRate(),
-		GetClipsMetaDataTag_IsCineACineCameraRecording(),
-		GetClipsMetaDataTag_IsNoGood(),
-		GetClipsMetaDataTag_IsFlagged(),
-		GetClipsMetaDataTag_FavoriteLevel(),
-		GetClipsMetaDataTag_IsCreatedFromVCam(),
 		GetClipsMetaDataTag_PostSmoothLevel()
 	};
 }
 
+PRAGMA_DISABLE_DEPRECATION_WARNINGS
 void UVirtualCameraClipsMetaData::ExtendAssetRegistryTags(FAssetRegistryTagsContext Context) const
 {
 	IMovieSceneMetaDataInterface::ExtendAssetRegistryTags(Context);
-
-	FString IsSelected = bIsSelected ? "true" : "false";
-	FString IsRecordedFromACineCameraActor = bIsACineCameraRecording ? "true" : "false";
-
-	Context.AddTag(FAssetRegistryTag(AssetRegistryTag_FocalLength, FString::FromInt(FocalLength), FAssetRegistryTag::ETagType::TT_Numerical, FAssetRegistryTag::TD_None));
-	Context.AddTag(FAssetRegistryTag(AssetRegistryTag_bIsSelected, IsSelected, FAssetRegistryTag::ETagType::TT_Alphabetical, FAssetRegistryTag::TD_None));
-	Context.AddTag(FAssetRegistryTag(AssetRegistryTag_RecordedLevelName, RecordedLevelName, FAssetRegistryTag::ETagType::TT_Alphabetical, FAssetRegistryTag::TD_None));
-	Context.AddTag(FAssetRegistryTag(AssetRegistryTag_FrameCountStart, FString::FromInt(FrameCountStart), FAssetRegistryTag::ETagType::TT_Numerical, FAssetRegistryTag::TD_None));
-	Context.AddTag(FAssetRegistryTag(AssetRegistryTag_FrameCountEnd, FString::FromInt(FrameCountEnd), FAssetRegistryTag::ETagType::TT_Numerical, FAssetRegistryTag::TD_None));
-	Context.AddTag(FAssetRegistryTag(AssetRegistryTag_LengthInFrames, FString::FromInt(LengthInFrames), FAssetRegistryTag::ETagType::TT_Numerical, FAssetRegistryTag::TD_None));
-	Context.AddTag(FAssetRegistryTag(AssetRegistryTag_DisplayRate, DisplayRate.ToPrettyText().ToString(), FAssetRegistryTag::ETagType::TT_Alphabetical, FAssetRegistryTag::TD_None));
-	Context.AddTag(FAssetRegistryTag(AssetRegistryTag_bIsACineCameraRecording, IsRecordedFromACineCameraActor, FAssetRegistryTag::ETagType::TT_Alphabetical, FAssetRegistryTag::TD_None));
-	Context.AddTag(FAssetRegistryTag(AssetRegistryTag_bIsNoGood, FString::FromInt(bIsNoGood), FAssetRegistryTag::ETagType::TT_Numerical, FAssetRegistryTag::TD_None));
-	Context.AddTag(FAssetRegistryTag(AssetRegistryTag_bIsFlagged, FString::FromInt(bIsFlagged), FAssetRegistryTag::ETagType::TT_Numerical, FAssetRegistryTag::TD_None));
-	Context.AddTag(FAssetRegistryTag(AssetRegistryTag_FavoriteLevel, FString::FromInt(FavoriteLevel), FAssetRegistryTag::ETagType::TT_Alphabetical, FAssetRegistryTag::TD_None));
-	Context.AddTag(FAssetRegistryTag(AssetRegistryTag_bIsCreatedFromVCam, FString::FromInt(bIsCreatedFromVCam), FAssetRegistryTag::ETagType::TT_Numerical, FAssetRegistryTag::TD_None));
+	
 	Context.AddTag(FAssetRegistryTag(AssetRegistryTag_PostSmoothLevel, FString::FromInt(PostSmoothLevel), FAssetRegistryTag::ETagType::TT_Numerical, FAssetRegistryTag::TD_None));
+
+	// These may have been migrated
+#if WITH_EDITOR
+	const ULevelSequence* OwningSequence = Cast<const ULevelSequence>(Context.GetObject());
+	if (!ULevelSequenceShotMetaDataLibrary::HasIsNoGood(OwningSequence))
+	{
+		Context.AddTag(FAssetRegistryTag(AssetRegistryTag_bIsNoGood, FString::FromInt(bIsNoGood), FAssetRegistryTag::ETagType::TT_Numerical, FAssetRegistryTag::TD_None));
+	}
+	if (!ULevelSequenceShotMetaDataLibrary::HasIsFlagged(OwningSequence))
+	{
+		Context.AddTag(FAssetRegistryTag(AssetRegistryTag_bIsFlagged, FString::FromInt(bIsFlagged), FAssetRegistryTag::ETagType::TT_Numerical, FAssetRegistryTag::TD_None));
+	}
+	if (!ULevelSequenceShotMetaDataLibrary::HasFavoriteRating(OwningSequence))
+	{
+		Context.AddTag(FAssetRegistryTag(AssetRegistryTag_FavoriteLevel, FString::FromInt(FavoriteLevel), FAssetRegistryTag::ETagType::TT_Alphabetical, FAssetRegistryTag::TD_None));
+	}
+#endif
 }
+PRAGMA_ENABLE_DEPRECATION_WARNINGS
 
 #if WITH_EDITOR
 void UVirtualCameraClipsMetaData::ExtendAssetRegistryTagMetaData(TMap<FName, FAssetRegistryTagMetadata>& OutMetadata) const
-{
-	OutMetadata.Add(AssetRegistryTag_FocalLength, FAssetRegistryTagMetadata()
-		.SetDisplayName(NSLOCTEXT("ClipsMetaData", "FocalLength_Label", "FocalLength"))
-		.SetTooltip(NSLOCTEXT("ClipsMetaData", "FocalLength_Tip", "The focal length that this level sequence was recorded with"))
-	);
-
-	OutMetadata.Add(AssetRegistryTag_bIsSelected, FAssetRegistryTagMetadata()
-		.SetDisplayName(NSLOCTEXT("ClipsMetaData", "IsSelected_Label", "bIsSelected"))
-		.SetTooltip(NSLOCTEXT("ClipsMetaData", "IsSelected_Tip", "If this level sequence is marked as a selected take"))
-	);
-
-	OutMetadata.Add(AssetRegistryTag_RecordedLevelName, FAssetRegistryTagMetadata()
-		.SetDisplayName(NSLOCTEXT("ClipsMetaData", "RecordedLavelName_Label", "RecordedLevelName"))
-		.SetTooltip(NSLOCTEXT("ClipsMetaData", "RecordedLavelName_Tip", "The name of the level this sequence was recorded"))
-	);
-
-	OutMetadata.Add(AssetRegistryTag_FrameCountStart, FAssetRegistryTagMetadata()
-		.SetDisplayName(NSLOCTEXT("ClipsMetaData", "FrameCountStart_Label", "FrameCountStart"))
-		.SetTooltip(NSLOCTEXT("ClipsMetaData", "FrameCountStart_Tip", "The first frame of the level sequence. "))
-	);
-
-	OutMetadata.Add(AssetRegistryTag_FrameCountEnd, FAssetRegistryTagMetadata()
-		.SetDisplayName(NSLOCTEXT("ClipsMetaData", "FrameCountEnd_Label", "FrameCountEnd"))
-		.SetTooltip(NSLOCTEXT("ClipsMetaData", "FrameCountEnd_Tip", "The last frame of the level sequence."))
-	);
-	OutMetadata.Add(AssetRegistryTag_LengthInFrames, FAssetRegistryTagMetadata()
-		.SetDisplayName(NSLOCTEXT("ClipsMetaData", "LengthInFrames_Label", "LengthInFrames"))
-		.SetTooltip(NSLOCTEXT("ClipsMetaData", "LengthInFrames_Tip", "The length in frames of the level sequence."))
-	);
-	OutMetadata.Add(AssetRegistryTag_DisplayRate, FAssetRegistryTagMetadata()
-		.SetDisplayName(NSLOCTEXT("ClipsMetaData", "DisplayRate_Label", "DisplayRate"))
-		.SetTooltip(NSLOCTEXT("ClipsMetaData", "DisplayRate_Tip", "The display rate of level sequence."))
-	);
-
-	OutMetadata.Add(AssetRegistryTag_bIsACineCameraRecording, FAssetRegistryTagMetadata()
-		.SetDisplayName(NSLOCTEXT("ClipsMetaData", "IsACineCameraRecording_Label", "bIsACineCameraRecording"))
-		.SetTooltip(NSLOCTEXT("ClipsMetaData", "IsACineCameraRecording_Tip", "If the level sequence was recorded with a CineCameraActor."))
-	);
-}
+{}
 #endif
 
+PRAGMA_DISABLE_DEPRECATION_WARNINGS
 float UVirtualCameraClipsMetaData::GetFocalLength() const 
 {
-	return FocalLength; 
+#if WITH_EDITOR
+	return FocalLength;
+#else
+	return {};
+#endif
 }
+PRAGMA_ENABLE_DEPRECATION_WARNINGS
 
+PRAGMA_DISABLE_DEPRECATION_WARNINGS
 bool UVirtualCameraClipsMetaData::GetSelected() const
 {
+#if WITH_EDITOR
 	return bIsSelected; 
+#else
+return {};
+#endif
 }
+PRAGMA_ENABLE_DEPRECATION_WARNINGS
 
+PRAGMA_DISABLE_DEPRECATION_WARNINGS
 FString UVirtualCameraClipsMetaData::GetRecordedLevelName() const
 {
+#if WITH_EDITOR
 	return RecordedLevelName; 
+#else
+return {};
+#endif
 }
+PRAGMA_ENABLE_DEPRECATION_WARNINGS
 
+PRAGMA_DISABLE_DEPRECATION_WARNINGS
 int UVirtualCameraClipsMetaData::GetFrameCountStart() const
 {
+#if WITH_EDITOR
 	return FrameCountStart; 
+#else
+return {};
+#endif
 }
+PRAGMA_ENABLE_DEPRECATION_WARNINGS
 
+PRAGMA_DISABLE_DEPRECATION_WARNINGS
 int UVirtualCameraClipsMetaData::GetFrameCountEnd() const
 {
+#if WITH_EDITOR
 	return FrameCountEnd; 
+#else
+return {};
+#endif
 }
+PRAGMA_ENABLE_DEPRECATION_WARNINGS
 
+PRAGMA_DISABLE_DEPRECATION_WARNINGS
 int UVirtualCameraClipsMetaData::GetLengthInFrames()
 {
+#if WITH_EDITOR
 	return LengthInFrames;
+#else
+return {};
+#endif
 }
+PRAGMA_ENABLE_DEPRECATION_WARNINGS
 
+PRAGMA_DISABLE_DEPRECATION_WARNINGS
 FFrameRate UVirtualCameraClipsMetaData::GetDisplayRate()
 {
+#if WITH_EDITOR
 	return DisplayRate;
+#else
+return {};
+#endif
 }
+PRAGMA_ENABLE_DEPRECATION_WARNINGS
 
+PRAGMA_DISABLE_DEPRECATION_WARNINGS
 bool UVirtualCameraClipsMetaData::GetIsACineCameraRecording() const
 {
+#if WITH_EDITOR
 	return bIsACineCameraRecording; 
+#else
+return {};
+#endif
+}
+PRAGMA_ENABLE_DEPRECATION_WARNINGS
+
+PRAGMA_DISABLE_DEPRECATION_WARNINGS
+bool UVirtualCameraClipsMetaData::GetIsNoGood() const
+{
+#if WITH_EDITOR
+	return bIsNoGood;
+#else
+return {};
+#endif
+}
+PRAGMA_ENABLE_DEPRECATION_WARNINGS
+
+PRAGMA_DISABLE_DEPRECATION_WARNINGS
+bool UVirtualCameraClipsMetaData::GetIsFlagged() const
+{
+#if WITH_EDITOR
+	return bIsFlagged;
+#else
+return {};
+#endif
+}
+PRAGMA_ENABLE_DEPRECATION_WARNINGS
+
+PRAGMA_DISABLE_DEPRECATION_WARNINGS
+int32 UVirtualCameraClipsMetaData::GetFavoriteLevel() const
+{
+#if WITH_EDITOR
+	return FavoriteLevel;
+#else
+return {};
+#endif
 }
+PRAGMA_ENABLE_DEPRECATION_WARNINGS
 
+PRAGMA_DISABLE_DEPRECATION_WARNINGS
 void UVirtualCameraClipsMetaData::SetFocalLength(float InFocalLength) 
 {
+#if WITH_EDITOR
 	FocalLength = InFocalLength;
+#endif
 }
+PRAGMA_ENABLE_DEPRECATION_WARNINGS
 
+PRAGMA_DISABLE_DEPRECATION_WARNINGS
 void UVirtualCameraClipsMetaData::SetSelected(bool bInSelected)
 {
+#if WITH_EDITOR
 	bIsSelected = bInSelected;
+#endif
 }
+PRAGMA_ENABLE_DEPRECATION_WARNINGS
 
+PRAGMA_DISABLE_DEPRECATION_WARNINGS
 void UVirtualCameraClipsMetaData::SetRecordedLevelName(FString InLevelName)
 {
+#if WITH_EDITOR
 	RecordedLevelName = InLevelName;
+#endif
 }
+PRAGMA_ENABLE_DEPRECATION_WARNINGS
 
+PRAGMA_DISABLE_DEPRECATION_WARNINGS
 void UVirtualCameraClipsMetaData::SetFrameCountStart(int InFrame)
 {
+#if WITH_EDITOR
 	FrameCountStart = InFrame; 
+#endif
 }
+PRAGMA_ENABLE_DEPRECATION_WARNINGS
 
+PRAGMA_DISABLE_DEPRECATION_WARNINGS
 void UVirtualCameraClipsMetaData::SetFrameCountEnd(int InFrame)
 {
+#if WITH_EDITOR
 	FrameCountEnd = InFrame; 
+#endif
 }
+PRAGMA_ENABLE_DEPRECATION_WARNINGS
 
+PRAGMA_DISABLE_DEPRECATION_WARNINGS
 void UVirtualCameraClipsMetaData::SetLengthInFrames(int InLength)
 {
+#if WITH_EDITOR
 	LengthInFrames = InLength; 
+#endif
 }
+PRAGMA_ENABLE_DEPRECATION_WARNINGS
 
+PRAGMA_DISABLE_DEPRECATION_WARNINGS
 void UVirtualCameraClipsMetaData::SetDisplayRate(FFrameRate InDisplayRate)
 {
+#if WITH_EDITOR
 	DisplayRate = InDisplayRate; 
+#endif
 }
+PRAGMA_ENABLE_DEPRECATION_WARNINGS
 
+PRAGMA_DISABLE_DEPRECATION_WARNINGS
 void UVirtualCameraClipsMetaData::SetIsACineCameraRecording(bool bInIsACineCameraRecording)
 {
-	bIsACineCameraRecording = bInIsACineCameraRecording; 
+#if WITH_EDITOR
+	bIsACineCameraRecording = bInIsACineCameraRecording;
+#endif
 }
-
+PRAGMA_ENABLE_DEPRECATION_WARNINGS
\ No newline at end of file
diff --git a/Engine/Plugins/VirtualProduction/VirtualCamera/Source/VirtualCamera/Public/LevelSequence/VCamTakesMetaDataMigration.h b/Engine/Plugins/VirtualProduction/VirtualCamera/Source/VirtualCamera/Public/LevelSequence/VCamTakesMetaDataMigration.h
new file mode 100644
--- /dev/null
+++ b/Engine/Plugins/VirtualProduction/VirtualCamera/Source/VirtualCamera/Public/LevelSequence/VCamTakesMetaDataMigration.h
@@ -0,0 +1,88 @@
+// Copyright Epic Games, Inc. All Rights Reserved.
+
+#pragma once
+
+#include "AssetRegistry/AssetData.h"
+#include "Kismet/BlueprintFunctionLibrary.h"
+#include "VCamTakesMetaDataMigration.generated.h"
+
+class ULevelSequence;
+
+/** Helps VCam Blueprints migrate. UVirtualCameraClipsMetaData */
+UCLASS()
+class UVCamTakesMetaDataMigration : public UBlueprintFunctionLibrary
+{
+	GENERATED_BODY()
+public:
+
+	/**
+	 * Gets whether the user has marked this take as no good.
+	 * Does no migration. Looks for new data source first and falls back to legacy data.
+	 * @return Whether there is any data (legacy or new) stored.
+	 */
+	UFUNCTION(BlueprintPure, Category = "VirtualCamera", meta = (DisplayName = "Get Is No Good (Migration)"))
+	static bool GetIsNoGood(ULevelSequence* InLevelSequence, bool& bOutNoGood);
+	/**
+	 * Gets whether this was flagged by a user.
+	 * Does no migration. Looks for new data source first and falls back to legacy data.
+	 * @return Whether there is any data (legacy or new) stored.
+	 */
+	UFUNCTION(BlueprintPure, Category = "VirtualCamera", meta = (DisplayName = "Get Is Flagged (Migration)"))
+	static bool GetIsFlagged(ULevelSequence* InLevelSequence, bool& bOutIsFlagged);
+	/**
+	 * Gets the favorite rating. The favorite rating is like a star rating, usually 1-3 if it was rated.
+	 * Does no migration. Looks for new data source first and falls back to legacy data.
+	 * @return Whether there is any data (legacy or new) stored.
+	 */
+	UFUNCTION(BlueprintPure, Category = "VirtualCamera", meta = (DisplayName = "Get Favorite Level (Migration)"))
+	static bool GetFavoriteLevel(ULevelSequence* InLevelSequence, int32& OutFavoriteLevel);
+
+	/**
+	 * Gets whether the user has marked this take as no good.
+	 * Does no migration. Looks for new data source first and falls back to legacy data.
+	 * @return Whether there is any data (legacy or new) stored.
+	 */
+	UFUNCTION(BlueprintPure, Category = "VirtualCamera", meta = (DisplayName = "Get Is No Good By Asset Data (Migration)"))
+	static bool GetIsNoGoodByAssetData(const FAssetData& InAssetData, bool& bOutNoGood);
+	/**
+	 * Gets whether this was flagged by a user.
+	 * Does no migration. Looks for new data source first and falls back to legacy data.
+	 * @return Whether there is any data (legacy or new) stored.
+	 */
+	UFUNCTION(BlueprintPure, Category = "VirtualCamera", meta = (DisplayName = "Get Is Flagged By Asset Data (Migration)"))
+	static bool GetIsFlaggedByAssetData(const FAssetData& InAssetData, bool& bOutIsFlagged);
+	/**
+	 * Gets the favorite rating. The favorite rating is like a star rating, usually 1-3 if it was rated.
+	 * Does no migration. Looks for new data source first and falls back to legacy data.
+	 * @return Whether there is any data (legacy or new) stored.
+	 */
+	UFUNCTION(BlueprintPure, Category = "VirtualCamera", meta = (DisplayName = "Get Favorite Level By Asset Data (Migration)"))
+	static bool GetFavoriteLevelByAssetData(const FAssetData& InAssetData, int32& OutFavoriteLevel);
+
+	/** Sets whether the user has marked this take as no good. Always writes the data to the new data source. */
+	UFUNCTION(BlueprintCallable, Category = "VirtualCamera", meta = (DisplayName = "Set Is No Good (Migration)"))
+	static void SetIsNoGood(ULevelSequence* InLevelSequence, bool bInIsNoGood);
+	/** Sets whether this was flagged by a user. Always writes the data to the new data source. */
+	UFUNCTION(BlueprintCallable, Category = "VirtualCamera", meta = (DisplayName = "Set Is Flagged (Migration)"))
+	static void SetIsFlagged(ULevelSequence* InLevelSequence, bool bInIsFlagged);
+	/** Sets the favorite rating. The favorite rating is like a star rating, usually 1-3 if it was rated. Always writes the data to the new data source. */
+	UFUNCTION(BlueprintCallable, Category = "VirtualCamera", meta = (DisplayName = "Set Favorite Level (Migration)"))
+	static void SetFavoriteLevel(ULevelSequence* InLevelSequence, int32 InFavoriteLevel);
+
+	/** @return Whether the passed in level sequence contains the legacy UVirtualCameraClipsMetaData */
+	UFUNCTION(BlueprintPure, Category = "VirtualCamera", meta = (DisplayName = "Needs Migration For VCam TakesMeta Data (Migration)"))
+	static bool NeedsToMigrateVCamMetaData(ULevelSequence* InLevelSequence);
+	/** @return Whether the passed in level sequence contains the legacy UVirtualCameraClipsMetaData */
+	UFUNCTION(BlueprintPure, Category = "VirtualCamera", meta = (DisplayName = "Needs Migration For VCam TakesMeta Data By Asset Data (Migration)"))
+	static bool NeedsToMigrateVCamMetaDataByAssetData(const FAssetData& InAssetData);
+	
+	/** Migrates UVirtualCameraClipsMetaData to UMovieSceneShotMetaData. */
+	UFUNCTION(BlueprintCallable, Category = "VirtualCamera", meta = (DisplayName = "Migrate VCam Takes Meta Data (Migration)"))
+	static void MigrateVCamTakesMetaData(ULevelSequence* InLevelSequence);
+
+	UFUNCTION(BlueprintCallable, Category = "VirtualCamera")
+	static bool GetAutoMigrateAccessedLevelSequencesCVar();
+	
+	UFUNCTION(BlueprintCallable, Category = "VirtualCamera")
+	static void SLOW_MigrateAllVCamTakesMetaDataInProject();
+};
diff --git a/Engine/Plugins/VirtualProduction/VirtualCamera/Source/VirtualCamera/Public/LevelSequence/VirtualCameraClipsMetaData.h b/Engine/Plugins/VirtualProduction/VirtualCamera/Source/VirtualCamera/Public/LevelSequence/VirtualCameraClipsMetaData.h
--- a/Engine/Plugins/VirtualProduction/VirtualCamera/Source/VirtualCamera/Public/LevelSequence/VirtualCameraClipsMetaData.h
+++ b/Engine/Plugins/VirtualProduction/VirtualCamera/Source/VirtualCamera/Public/LevelSequence/VirtualCameraClipsMetaData.h
@@ -2,61 +2,73 @@
 
 #pragma once
 
-
 #include "UObject/Object.h"
 #include "IMovieSceneMetaData.h"
 #include "Misc/FrameRate.h"
 #include "VirtualCameraClipsMetaData.generated.h"
 
-
 /**
  * Clips meta-data that is stored on ULevelSequence assets that are recorded through the virtual camera. 
- * Meta-data is retrieved through ULevelSequence::FindMetaData<UVirtualCameraClipsMetaData>()
+ * Meta-data is retrieved through ULevelSequence::FindMetaData<UVirtualCameraClipsMetaData>().
  */
-
 UCLASS(BlueprintType)
-class UVirtualCameraClipsMetaData : public UObject, public IMovieSceneMetaDataInterface
+class UVirtualCameraClipsMetaData
+	: public UObject
+	, public IMovieSceneMetaDataInterface
 {
 public: 
 	GENERATED_BODY()
+	
 	UVirtualCameraClipsMetaData(const FObjectInitializer& ObjInit);
 
 public: 
 
 	/** The asset registry tag that contains the focal length for this meta-data */
+	UE_DEPRECATED(5.6, "Data was removed.")
 	static const FName AssetRegistryTag_FocalLength;
 
 	/** The asset registry tag that contains if the selected state for this meta-data */
+	UE_DEPRECATED(5.6, "Data was removed.")
 	static const FName AssetRegistryTag_bIsSelected;
 
 	/** The asset registry tag that contains the recorded level name for this meta-data */
+	UE_DEPRECATED(5.6, "Data was removed.")
 	static const FName AssetRegistryTag_RecordedLevelName;
 
 	/** The asset registry tag that contains the FrameCountStart in for this meta-data */
+	UE_DEPRECATED(5.6, "Data was removed.")
 	static const FName AssetRegistryTag_FrameCountStart;
 
 	/** The asset registry tag that contains the FrameCountEnd out for this meta-data */
+	UE_DEPRECATED(5.6, "Data was removed.")
 	static const FName AssetRegistryTag_FrameCountEnd;
 
 	/** The asset registry tag that contains the LengthInFrames out for this meta-data */
+	UE_DEPRECATED(5.6, "Data was removed.")
 	static const FName AssetRegistryTag_LengthInFrames; 
 
 	/** The asset registry tag that contains the FrameCountEnd out for this meta-data */
+	UE_DEPRECATED(5.6, "Data was removed.")
 	static const FName AssetRegistryTag_DisplayRate;
 
 	/** The asset registry tag that contains whether the clip was recorded with a CineCamera for this meta-data */
+	UE_DEPRECATED(5.6, "Data was removed.")
 	static const FName AssetRegistryTag_bIsACineCameraRecording;
 	
 	/** The asset registry tag that contains whether this take is good or not */
+	UE_DEPRECATED(5.6, "Use ULevelSequenceShotMetaDataLibrary::GetIsNoGoodAssetTag instead.")
 	static const FName AssetRegistryTag_bIsNoGood;
 
 	/** The asset registry tag that contains whether this was flagged by a user */
+	UE_DEPRECATED(5.6, "Use ULevelSequenceShotMetaDataLibrary::GetIsFlaggedAssetTag instead.")
 	static const FName AssetRegistryTag_bIsFlagged;
 
 	/** The asset registry tag that contains its favorite status */
+	UE_DEPRECATED(5.6, "Use ULevelSequenceShotMetaDataLibrary::GetFavoriteRatingAssetTag instead.")
 	static const FName AssetRegistryTag_FavoriteLevel;
 
 	/** The asset registry tag that contains whether it was created from a VCam */
+	UE_DEPRECATED(5.6, "Data was removed.")
 	static const FName AssetRegistryTag_bIsCreatedFromVCam;
 
 	/* The asset registry tag that contains the post smooth level*/
@@ -65,56 +77,68 @@ public:
 public:
 
 	/** The asset registry tag that contains the focal length for this meta-data */
-	UFUNCTION(BlueprintPure, Category = "VirtualCamera|Clips")
-	static FName GetClipsMetaDataTag_FocalLength() { return AssetRegistryTag_FocalLength; }
+	UE_DEPRECATED(5.6, "Data was removed.")
+	UFUNCTION(BlueprintPure, Category = "VirtualCamera|Clips", meta=(DeprecatedFunction, DeprecationMessage="Data was removed."))
+	static FName GetClipsMetaDataTag_FocalLength();
 	
 	/** The asset registry tag that contains if the selected state for this meta-data */
-	UFUNCTION(BlueprintPure, Category = "VirtualCamera|Clips")
-	static FName GetClipsMetaDataTag_IsSelected() { return AssetRegistryTag_bIsSelected; }
+	UE_DEPRECATED(5.6, "Data was removed.")
+	UFUNCTION(BlueprintPure, Category = "VirtualCamera|Clips", meta=(DeprecatedFunction, DeprecationMessage="Data was removed."))
+	static FName GetClipsMetaDataTag_IsSelected();
 	
 	/** The asset registry tag that contains the recorded level name for this meta-data */
-	UFUNCTION(BlueprintPure, Category = "VirtualCamera|Clips")
-	static FName GetClipsMetaDataTag_RecordedLevel() { return AssetRegistryTag_RecordedLevelName; }
+	UE_DEPRECATED(5.6, "Data was removed.")
+	UFUNCTION(BlueprintPure, Category = "VirtualCamera|Clips", meta=(DeprecatedFunction, DeprecationMessage="Data was removed."))
+	static FName GetClipsMetaDataTag_RecordedLevel();
 	
 	/** The asset registry tag that contains the FrameCountStart in for this meta-data */
-	UFUNCTION(BlueprintPure, Category = "VirtualCamera|Clips")
-	static FName GetClipsMetaDataTag_FrameCountStart() { return AssetRegistryTag_FrameCountStart; }
+	UE_DEPRECATED(5.6, "Data was removed.")
+	UFUNCTION(BlueprintPure, Category = "VirtualCamera|Clips", meta=(DeprecatedFunction, DeprecationMessage="Data was removed."))
+	static FName GetClipsMetaDataTag_FrameCountStart();
 	
 	/** The asset registry tag that contains the FrameCountEnd out for this meta-data */
-	UFUNCTION(BlueprintPure, Category = "VirtualCamera|Clips")
-	static FName GetClipsMetaDataTag_FrameCountEnd() { return AssetRegistryTag_FrameCountEnd; }
+	UE_DEPRECATED(5.6, "Data was removed.")
+	UFUNCTION(BlueprintPure, Category = "VirtualCamera|Clips", meta=(DeprecatedFunction, DeprecationMessage="Data was removed."))
+	static FName GetClipsMetaDataTag_FrameCountEnd();
 	
 	/** The asset registry tag that contains the LengthInFrames out for this meta-data */
-	UFUNCTION(BlueprintPure, Category = "VirtualCamera|Clips")
-	static FName GetClipsMetaDataTag_LengthInFrames() { return AssetRegistryTag_LengthInFrames; }
+	UE_DEPRECATED(5.6, "Data was removed.")
+	UFUNCTION(BlueprintPure, Category = "VirtualCamera|Clips", meta=(DeprecatedFunction, DeprecationMessage="Data was removed."))
+	static FName GetClipsMetaDataTag_LengthInFrames();
 	
 	/** The asset registry tag that contains the FrameCountEnd out for this meta-data */
-	UFUNCTION(BlueprintPure, Category = "VirtualCamera|Clips")
-	static FName GetClipsMetaDataTag_DisplayRate() { return AssetRegistryTag_DisplayRate; }
+	UE_DEPRECATED(5.6, "Data was removed.")
+	UFUNCTION(BlueprintPure, Category = "VirtualCamera|Clips", meta=(DeprecatedFunction, DeprecationMessage="Data was removed."))
+	static FName GetClipsMetaDataTag_DisplayRate();
 	
 	/** The asset registry tag that contains whether the clip was recorded with a CineCamera for this meta-data */
-	UFUNCTION(BlueprintPure, Category = "Virtual Camera|Clips")
-	static FName GetClipsMetaDataTag_IsCineACineCameraRecording() { return AssetRegistryTag_bIsACineCameraRecording; }
+	UE_DEPRECATED(5.6, "Data was removed.")
+	UFUNCTION(BlueprintPure, Category = "Virtual Camera|Clips", meta=(DeprecatedFunction, DeprecationMessage="Data was removed."))
+	static FName GetClipsMetaDataTag_IsCineACineCameraRecording();
 	
 	/** The asset registry tag that contains whether this take is good or not */
-	UFUNCTION(BlueprintPure, Category = "VirtualCamera|Clips")
-	static FName GetClipsMetaDataTag_IsNoGood() { return AssetRegistryTag_bIsNoGood; }
+	UE_DEPRECATED(5.6, "Use ULevelSequenceShotMetaDataLibrary::GetIsNoGoodAssetTag instead.")
+	UFUNCTION(BlueprintPure, Category = "VirtualCamera|Clips", meta=(DeprecatedFunction, DeprecationMessage="Use ULevelSequenceShotMetaDataLibrary::GetIsNoGoodAssetTag instead."))
+	static FName GetClipsMetaDataTag_IsNoGood();
 	
 	/** The asset registry tag that contains whether this was flagged by a user */
-	UFUNCTION(BlueprintPure, Category = "VirtualCamera|Clips")
-	static FName GetClipsMetaDataTag_IsFlagged() { return AssetRegistryTag_bIsFlagged; }
+	UE_DEPRECATED(5.6, "Use ULevelSequenceShotMetaDataLibrary::GetIsFlaggedAssetTag instead.")
+	UFUNCTION(BlueprintPure, Category = "VirtualCamera|Clips", meta=(DeprecatedFunction, DeprecationMessage="Use ULevelSequenceShotMetaDataLibrary::GetIsFlaggedAssetTag instead."))
+	static FName GetClipsMetaDataTag_IsFlagged();
 	
 	/** The asset registry tag that contains its favorite status */
-	UFUNCTION(BlueprintPure, Category = "VirtualCamera|Clips")
-	static FName GetClipsMetaDataTag_FavoriteLevel() { return AssetRegistryTag_FavoriteLevel; }
+	UE_DEPRECATED(5.6, "Use ULevelSequenceShotMetaDataLibrary::GetFavoriteRatingAssetTag instead.")
+	UFUNCTION(BlueprintPure, Category = "VirtualCamera|Clips", meta=(DeprecatedFunction, DeprecationMessage="Use ULevelSequenceShotMetaDataLibrary::GetFavoriteRatingAssetTag instead."))
+	static FName GetClipsMetaDataTag_FavoriteLevel();
 	
 	/** The asset registry tag that contains whether it was created from a VCam */
-	UFUNCTION(BlueprintPure, Category = "VirtualCamera|Clips")
-	static FName GetClipsMetaDataTag_IsCreatedFromVCam() { return AssetRegistryTag_bIsCreatedFromVCam; }
+	UE_DEPRECATED(5.6, "Data was removed.")
+	UFUNCTION(BlueprintPure, Category = "VirtualCamera|Clips", meta=(DeprecatedFunction, DeprecationMessage="Data was removed."))
+	static FName GetClipsMetaDataTag_IsCreatedFromVCam();
 
 	/** The asset registry tag that contains post smooth level */
 	UFUNCTION(BlueprintPure, Category = "VirtualCamera|Clips")
-	static FName GetClipsMetaDataTag_PostSmoothLevel() { return AssetRegistryTag_PostSmoothLevel; }
+	static FName GetClipsMetaDataTag_PostSmoothLevel();
 
 	/** Gets all asset registry tags */
 	UFUNCTION(BlueprintPure, Category = "VirtualCamera|Clips")
@@ -136,151 +160,188 @@ public:
 	/**
 	 * @return The focal length for this clip
 	 */
-	UFUNCTION(BlueprintCallable, Category = "Clips")
+	UE_DEPRECATED(5.6, "Data was removed.")
+	UFUNCTION(BlueprintCallable, Category = "Clips", meta=(DeprecatedFunction, DeprecationMessage="Data was removed."))
 	float GetFocalLength() const;
 
 	/**
 	* @return Whether or not the clip is selected.  
 	*/
-	UFUNCTION(BlueprintCallable, Category = "Clips")
+	UE_DEPRECATED(5.6, "Data was removed.")
+	UFUNCTION(BlueprintCallable, Category = "Clips", meta=(DeprecatedFunction, DeprecationMessage="Data was removed."))
 	bool GetSelected() const;
 
 	/**
 	* @return The name of the clip's recorded level. 
 	*/
-	UFUNCTION(BlueprintCallable, Category = "Clips")
+	UE_DEPRECATED(5.6, "Data was removed.")
+	UFUNCTION(BlueprintCallable, Category = "Clips", meta=(DeprecatedFunction, DeprecationMessage="Data was removed."))
 	FString GetRecordedLevelName() const;
 
 	/**
 	* @return The initial frame of the clip
 	*/
-	UFUNCTION(BlueprintCallable, Category = "Clips")
+	UE_DEPRECATED(5.6, "Data was removed.")
+	UFUNCTION(BlueprintCallable, Category = "Clips", meta=(DeprecatedFunction, DeprecationMessage="Data was removed."))
 	int GetFrameCountStart() const;
 
 	/**
 	* @return The final frame of the clip
 	*/
-	UFUNCTION(BlueprintCallable, Category = "Clips")
+	UE_DEPRECATED(5.6, "Data was removed.")
+	UFUNCTION(BlueprintCallable, Category = "Clips", meta=(DeprecatedFunction, DeprecationMessage="Data was removed."))
 	int GetFrameCountEnd() const;
 
 	/**
 	* @return The length in frames of the clip. 
 	*/
-	UFUNCTION(BlueprintCallable, Category = "Clips")
+	UE_DEPRECATED(5.6, "Data was removed.")
+	UFUNCTION(BlueprintCallable, Category = "Clips", meta=(DeprecatedFunction, DeprecationMessage="Data was removed."))
 	int GetLengthInFrames();
 
 	/**
 	* @return The display rate of the clip. 
 	*/
-	UFUNCTION(BlueprintCallable, Category = "Clips")
+	UE_DEPRECATED(5.6, "Data was removed.")
+	UFUNCTION(BlueprintCallable, Category = "Clips", meta=(DeprecatedFunction, DeprecationMessage="Data was removed."))
 	FFrameRate GetDisplayRate();
 
 	/**
 	* @return Whether the clip was recorded by a CineCameraActor
 	*/
-	UFUNCTION(BlueprintCallable, Category = "Clips")
+	UE_DEPRECATED(5.6, "Data was removed.")
+	UFUNCTION(BlueprintCallable, Category = "Clips", meta=(DeprecatedFunction, DeprecationMessage="Data was removed."))
 	bool GetIsACineCameraRecording() const;
 
+	UE_DEPRECATED(5.6, "Use ULevelSequenceShotMetaDataLibrary::GetIsNoGood or UVCamTakesMetaDataMigration instead")
+	bool GetIsNoGood() const;
+	UE_DEPRECATED(5.6, "Use ULevelSequenceShotMetaDataLibrary::GetIsFlagged or UVCamTakesMetaDataMigration instead")
+	bool GetIsFlagged() const;
+	UE_DEPRECATED(5.6, "Use ULevelSequenceShotMetaDataLibrary::GetFavoriteRank or UVCamTakesMetaDataMigration instead")
+	int32 GetFavoriteLevel() const;
+
 public:
 	/**
 	* Set the focal length associated with this clip. 
 	* @note: Used for tracking. Does not update the StreamedCameraComponent. 
 	*/
-	UFUNCTION(BlueprintCallable, Category = "Clips")
+	UE_DEPRECATED(5.6, "Data was removed.")
+	UFUNCTION(BlueprintCallable, Category = "Clips", meta=(DeprecatedFunction, DeprecationMessage="Data was removed."))
 	void SetFocalLength(float InFocalLength);
 	
 	/**
 	* Set if this clip is 'selected'
 	*/
-	UFUNCTION(BlueprintCallable, Category = "Clips")
+	UE_DEPRECATED(5.6, "Data was removed.")
+	UFUNCTION(BlueprintCallable, Category = "Clips", meta=(DeprecatedFunction, DeprecationMessage="Data was removed."))
 	void SetSelected(bool bInSelected);
 
 	/**
 	* Set the name of the level that the clip was recorded in. 
 	*/
-	UFUNCTION(BlueprintCallable, Category = "Clips")
+	UE_DEPRECATED(5.6, "Data was removed.")
+	UFUNCTION(BlueprintCallable, Category = "Clips", meta=(DeprecatedFunction, DeprecationMessage="Data was removed."))
 	void SetRecordedLevelName(FString InLevelName);
 
 	/**
 	* Set the initial frame of the clip used for calculating duration.
 	*/
-	UFUNCTION(BlueprintCallable, Category = "Clips")
+	UE_DEPRECATED(5.6, "Data was removed.")
+	UFUNCTION(BlueprintCallable, Category = "Clips", meta=(DeprecatedFunction, DeprecationMessage="Data was removed."))
 	void SetFrameCountStart(int InFrame); 
 
 	/**
 	* Set the final frame of the clip used for calculating duration.
 	*/
-	UFUNCTION(BlueprintCallable, Category = "Clips")
+	UE_DEPRECATED(5.6, "Data was removed.")
+	UFUNCTION(BlueprintCallable, Category = "Clips", meta=(DeprecatedFunction, DeprecationMessage="Data was removed."))
 	void SetFrameCountEnd(int InFrame);
 
 	/**
 	 * Set the length in frames of the clip used for AssetData calculations. 
 	 */
-	UFUNCTION(BlueprintCallable, Category = "Clips")
+	UE_DEPRECATED(5.6, "Data was removed.")
+	UFUNCTION(BlueprintCallable, Category = "Clips", meta=(DeprecatedFunction, DeprecationMessage="Data was removed."))
 	void SetLengthInFrames(int InLength);
 
 	/**
 	 * Set the DisplayRate of the clip used for AssetData calculations.
 	 */
-	UFUNCTION(BlueprintCallable, Category = "Clips")
+	UE_DEPRECATED(5.6, "Data was removed.")
+	UFUNCTION(BlueprintCallable, Category = "Clips", meta=(DeprecatedFunction, DeprecationMessage="Data was removed."))
 	void SetDisplayRate(FFrameRate InDisplayRate);
 
 	/**
 	 * Set if the clip was recorded by a CineCameraActor
 	 */
-	UFUNCTION(BlueprintCallable, Category = "Clips")
+	UE_DEPRECATED(5.6, "Data was removed.")
+	UFUNCTION(BlueprintCallable, Category = "Clips", meta=(DeprecatedFunction, DeprecationMessage="Data was removed."))
 	void SetIsACineCameraRecording(bool bInIsACineCameraRecording);
 
 private:
-	
+
+#if	WITH_EDITORONLY_DATA // These properties are deprecated
+	UE_DEPRECATED(5.6, "Data was removed.")
 	/** The focal length of the streamed camera used to record the take */
-	UPROPERTY(EditAnywhere, Category = "Clips", BlueprintReadWrite, meta = (AllowPrivateAccess = "true"))
-	float FocalLength;
+	UPROPERTY(EditAnywhere, Category = "Clips", BlueprintReadWrite, meta = (AllowPrivateAccess = "true", DeprecatedProperty, DeprecationMessage = "Data was removed."))
+	float FocalLength = 0.f;
 
 	/** Whether or not the take was marked as 'selected' */
-	UPROPERTY(EditAnywhere, Category = "Clips", BlueprintReadWrite, meta = (AllowPrivateAccess = "true"))
-	bool bIsSelected; 
+	UE_DEPRECATED(5.6, "Data was removed.")
+	UPROPERTY(EditAnywhere, Category = "Clips", BlueprintReadWrite, meta = (AllowPrivateAccess = "true", DeprecatedProperty, DeprecationMessage = "Data was removed."))
+	bool bIsSelected = false; 
 
 	/** The name of the level that the clip was recorded in */
-	UPROPERTY(EditAnywhere, Category = "Clips", BlueprintReadWrite, meta = (AllowPrivateAccess = "true"))
+	UE_DEPRECATED(5.6, "Data was removed.")
+	UPROPERTY(EditAnywhere, Category = "Clips", BlueprintReadWrite, meta = (AllowPrivateAccess = "true", DeprecatedProperty, DeprecationMessage = "Data was removed."))
 	FString RecordedLevelName; 
 
 	/** The initial frame of the clip used for calculating duration. */
-	UPROPERTY(EditAnywhere, Category = "Clips", BlueprintReadWrite, meta = (AllowPrivateAccess = "true"))
-	int FrameCountStart;
+	UE_DEPRECATED(5.6, "Data was removed.")
+	UPROPERTY(EditAnywhere, Category = "Clips", BlueprintReadWrite, meta = (AllowPrivateAccess = "true", DeprecatedProperty, DeprecationMessage = "Data was removed."))
+	int FrameCountStart = 0;
 
 	/** The last frame of the clip used for calculating duration. */
-	UPROPERTY(EditAnywhere, Category = "Clips", BlueprintReadWrite, meta = (AllowPrivateAccess = "true"))
-	int FrameCountEnd; 
+	UE_DEPRECATED(5.6, "Data was removed.")
+	UPROPERTY(EditAnywhere, Category = "Clips", BlueprintReadWrite, meta = (AllowPrivateAccess = "true", DeprecatedProperty, DeprecationMessage = "Data was removed."))
+	int FrameCountEnd = 0; 
 
 	/** The level sequence length in frames calculated from VirtualCameraSubsystem used for AssetData calculations */
-	UPROPERTY(EditAnywhere, Category = "Clips", BlueprintReadWrite, meta = (AllowPrivateAccess = "true"))
-	int LengthInFrames; 
+	UE_DEPRECATED(5.6, "Data was removed.")
+	UPROPERTY(EditAnywhere, Category = "Clips", BlueprintReadWrite, meta = (AllowPrivateAccess = "true", DeprecatedProperty, DeprecationMessage = "Data was removed."))
+	int LengthInFrames = 0; 
 
 	/** The display rate of the level sequence used for AssetData calculations. */
-	UPROPERTY(EditAnywhere, Category = "Clips", BlueprintReadWrite, meta = (AllowPrivateAccess = "true"))
+	UE_DEPRECATED(5.6, "Data was removed.")
+	UPROPERTY(EditAnywhere, Category = "Clips", BlueprintReadWrite, meta = (AllowPrivateAccess = "true", DeprecatedProperty, DeprecationMessage = "Data was removed."))
 	FFrameRate DisplayRate;
 
 	/** If the LevelSequence was recorded with a CineCameraActor, rather than a VirtualCameraActor */
-	UPROPERTY(EditAnywhere, Category = "Clips", BlueprintReadWrite, meta = (AllowPrivateAccess = "true"))
-	bool bIsACineCameraRecording;
+	UE_DEPRECATED(5.6, "Data was removed.")
+	UPROPERTY(EditAnywhere, Category = "Clips", BlueprintReadWrite, meta = (AllowPrivateAccess = "true", DeprecatedProperty, DeprecationMessage = "Data was removed."))
+	bool bIsACineCameraRecording = false;
 	
 	/** Whether this take is marked as good */
-	UPROPERTY(EditAnywhere, Category = "Clips", BlueprintReadWrite, meta = (AllowPrivateAccess = "true"))
+	UE_DEPRECATED(5.6, "Data was removed.")
+	UPROPERTY(EditAnywhere, Category = "Clips", BlueprintReadWrite, meta = (AllowPrivateAccess = "true", DeprecatedProperty, DeprecationMessage = "Use ULevelSequenceShotMetaDataLibrary or UVCamTakesMetaDataMigration instead"))
 	bool bIsNoGood = false;
 	
 	/** The asset registry tag that contains whether this was flagged by a user */
-	UPROPERTY(EditAnywhere, Category = "Clips", BlueprintReadWrite, meta = (AllowPrivateAccess = "true"))
+	UE_DEPRECATED(5.6, "Data was removed.")
+	UPROPERTY(EditAnywhere, Category = "Clips", BlueprintReadWrite, meta = (AllowPrivateAccess = "true", DeprecatedProperty, DeprecationMessage = "Use ULevelSequenceShotMetaDataLibrary or UVCamTakesMetaDataMigration instead"))
 	bool bIsFlagged = false;
 
 	/** The asset registry tag that contains its favorite status */
-	UPROPERTY(EditAnywhere, Category = "Clips", BlueprintReadWrite, meta = (AllowPrivateAccess = "true"))
+	UE_DEPRECATED(5.6, "Data was removed.")
+	UPROPERTY(EditAnywhere, Category = "Clips", BlueprintReadWrite, meta = (AllowPrivateAccess = "true", DeprecatedProperty, DeprecationMessage = "Use ULevelSequenceShotMetaDataLibrary or UVCamTakesMetaDataMigration instead"))
 	int32 FavoriteLevel = 0;
 
 	/** Whether the sequence was created from a VCam */
-	UPROPERTY(EditAnywhere, Category = "Clips", BlueprintReadWrite, meta = (AllowPrivateAccess = "true"))
+	UE_DEPRECATED(5.6, "Data was removed.")
+	UPROPERTY(EditAnywhere, Category = "Clips", BlueprintReadWrite, meta = (AllowPrivateAccess = "true", DeprecatedProperty, DeprecationMessage = "Data was removed."))
 	bool bIsCreatedFromVCam = true;
-
+#endif
+	
 	/* The asset registry tag that contains the post smooth level*/
 	UPROPERTY(EditAnywhere, Category = "Clips", BlueprintReadWrite, meta = (AllowPrivateAccess = "true"))
 	int32 PostSmoothLevel = 0;
diff --git a/Engine/Source/Runtime/LevelSequence/Private/LevelSequenceShotMetaDataLibrary.cpp b/Engine/Source/Runtime/LevelSequence/Private/LevelSequenceShotMetaDataLibrary.cpp
new file mode 100644
--- /dev/null
+++ b/Engine/Source/Runtime/LevelSequence/Private/LevelSequenceShotMetaDataLibrary.cpp
@@ -0,0 +1,178 @@
+// Copyright Epic Games, Inc. All Rights Reserved.
+
+#include "LevelSequenceShotMetaDataLibrary.h"
+
+#include "LevelSequence.h"
+#include "MetaData/MovieSceneShotMetaData.h"
+
+namespace UE::Sequencer::TakesMetaDataDetail
+{
+static UMovieSceneShotMetaData* FindMetaData(const ULevelSequence* InLevelSequence)
+{
+#if WITH_EDITOR 
+	return InLevelSequence ? InLevelSequence->FindMetaData<UMovieSceneShotMetaData>() : nullptr;
+#else
+	// If needed at runtime, we could consider wrapping UMovieSceneShotMetaData in UAssetUserData during cook time.
+	return nullptr;
+#endif
+}
+static UMovieSceneShotMetaData* FindOrAddMetaData(ULevelSequence* InLevelSequence)
+{
+#if WITH_EDITOR
+	return InLevelSequence ? InLevelSequence->FindOrAddMetaData<UMovieSceneShotMetaData>() : nullptr;
+#else
+	// If needed at runtime, we could consider wrapping UMovieSceneShotMetaData in UAssetUserData during cook time.
+	return nullptr;
+#endif
+}
+}
+
+bool ULevelSequenceShotMetaDataLibrary::GetIsNoGood(const ULevelSequence* InLevelSequence, bool& bOutNoGood)
+{
+	UMovieSceneShotMetaData* MetaData = UE::Sequencer::TakesMetaDataDetail::FindMetaData(InLevelSequence);
+	if (!MetaData || !MetaData->GetIsNoGood().IsSet())
+	{
+		return false;
+	}
+	
+	bOutNoGood = *MetaData->GetIsNoGood();
+	return true;
+}
+
+bool ULevelSequenceShotMetaDataLibrary::GetIsFlagged(const ULevelSequence* InLevelSequence, bool& bOutIsFlagged)
+{
+	UMovieSceneShotMetaData* MetaData = UE::Sequencer::TakesMetaDataDetail::FindMetaData(InLevelSequence);
+	if (!MetaData || !MetaData->GetIsFlagged().IsSet())
+	{
+		return false;
+	}
+	
+	bOutIsFlagged = *MetaData->GetIsFlagged();
+	return true;
+}
+
+bool ULevelSequenceShotMetaDataLibrary::GetFavoriteRating(const ULevelSequence* InLevelSequence, int32& OutFavoriteRating)
+{
+	UMovieSceneShotMetaData* MetaData = UE::Sequencer::TakesMetaDataDetail::FindMetaData(InLevelSequence);
+	if (!MetaData || !MetaData->GetFavoriteRating().IsSet())
+	{
+		return false;
+	}
+	
+	OutFavoriteRating = *MetaData->GetFavoriteRating();
+	return true;
+}
+
+bool ULevelSequenceShotMetaDataLibrary::HasIsNoGood(const ULevelSequence* InLevelSequence)
+{
+	bool Dummy = false;
+	return GetIsNoGood(InLevelSequence, Dummy);
+}
+
+bool ULevelSequenceShotMetaDataLibrary::HasIsFlagged(const ULevelSequence* InLevelSequence)
+{
+	bool Dummy = false;
+	return GetIsFlagged(InLevelSequence, Dummy);
+}
+
+bool ULevelSequenceShotMetaDataLibrary::HasFavoriteRating(const ULevelSequence* InLevelSequence)
+{
+	int32 Dummy = 0;
+	return GetFavoriteRating(InLevelSequence, Dummy);
+}
+
+void ULevelSequenceShotMetaDataLibrary::SetIsNoGood(ULevelSequence* InLevelSequence, bool bInIsNoGood)
+{
+	if (UMovieSceneShotMetaData* MetaData = UE::Sequencer::TakesMetaDataDetail::FindOrAddMetaData(InLevelSequence))
+	{
+		MetaData->SetIsNoGood(bInIsNoGood);
+	}
+}
+
+void ULevelSequenceShotMetaDataLibrary::SetIsFlagged(ULevelSequence* InLevelSequence, bool bInIsFlagged)
+{
+	if (UMovieSceneShotMetaData* MetaData = UE::Sequencer::TakesMetaDataDetail::FindOrAddMetaData(InLevelSequence))
+	{
+		MetaData->SetIsFlagged(bInIsFlagged);
+	}
+}
+
+void ULevelSequenceShotMetaDataLibrary::SetFavoriteRating(ULevelSequence* InLevelSequence, int32 InFavoriteRating)
+{
+	if (UMovieSceneShotMetaData* MetaData = UE::Sequencer::TakesMetaDataDetail::FindOrAddMetaData(InLevelSequence))
+	{
+		MetaData->SetFavoriteRating(InFavoriteRating);
+	}
+}
+
+void ULevelSequenceShotMetaDataLibrary::ClearIsNoGood(ULevelSequence* InLevelSequence)
+{
+	if (UMovieSceneShotMetaData* MetaData = UE::Sequencer::TakesMetaDataDetail::FindOrAddMetaData(InLevelSequence))
+	{
+		MetaData->ClearIsNoGood();
+	}
+}
+
+void ULevelSequenceShotMetaDataLibrary::ClearIsFlagged(ULevelSequence* InLevelSequence)
+{
+	if (UMovieSceneShotMetaData* MetaData = UE::Sequencer::TakesMetaDataDetail::FindOrAddMetaData(InLevelSequence))
+	{
+		MetaData->ClearIsFlagged();
+	}
+}
+
+void ULevelSequenceShotMetaDataLibrary::ClearFavoriteRating(ULevelSequence* InLevelSequence)
+{
+	if (UMovieSceneShotMetaData* MetaData = UE::Sequencer::TakesMetaDataDetail::FindOrAddMetaData(InLevelSequence))
+	{
+		MetaData->ClearFavoriteRating();
+	}
+}
+
+FName ULevelSequenceShotMetaDataLibrary::GetIsNoGoodAssetTag()
+{
+	return UMovieSceneShotMetaData::AssetRegistryTag_bIsNoGood;
+}
+
+FName ULevelSequenceShotMetaDataLibrary::GetIsFlaggedAssetTag()
+{
+	return UMovieSceneShotMetaData::AssetRegistryTag_bIsFlagged;
+}
+
+FName ULevelSequenceShotMetaDataLibrary::GetFavoriteRatingAssetTag()
+{
+	return UMovieSceneShotMetaData::AssetRegistryTag_FavoriteRating;
+}
+
+bool ULevelSequenceShotMetaDataLibrary::GetIsNoGoodByAssetData(const FAssetData& InAssetData, bool& bOutNoGood)
+{
+	return UMovieSceneShotMetaData::GetIsNoGoodByAssetData(InAssetData, bOutNoGood);
+}
+
+bool ULevelSequenceShotMetaDataLibrary::GetIsFlaggedByAssetData(const FAssetData& InAssetData, bool& bOutIsFlagged)
+{
+	return UMovieSceneShotMetaData::GetIsFlaggedByAssetData(InAssetData, bOutIsFlagged);
+}
+
+bool ULevelSequenceShotMetaDataLibrary::GetFavoriteRatingByAssetData(const FAssetData& InAssetData, int32& OutFavoriteRating)
+{
+	return UMovieSceneShotMetaData::GetFavoriteRatingByAssetData(InAssetData, OutFavoriteRating);
+}
+
+bool ULevelSequenceShotMetaDataLibrary::HasIsNoGoodByAssetData(const FAssetData& InAssetData)
+{
+	bool bDummy = false;
+	return GetIsNoGoodByAssetData(InAssetData, bDummy);
+}
+
+bool ULevelSequenceShotMetaDataLibrary::HasIsFlaggedByAssetData(const FAssetData& InAssetData)
+{
+	bool bDummy = false;
+	return GetIsFlaggedByAssetData(InAssetData, bDummy);
+}
+
+bool ULevelSequenceShotMetaDataLibrary::HasFavoriteRatingByAssetData(const FAssetData& InAssetData)
+{
+	int32 Dummy = 0;
+	return GetFavoriteRatingByAssetData(InAssetData, Dummy);
+}
diff --git a/Engine/Source/Runtime/LevelSequence/Public/LevelSequenceShotMetaDataLibrary.h b/Engine/Source/Runtime/LevelSequence/Public/LevelSequenceShotMetaDataLibrary.h
new file mode 100644
--- /dev/null
+++ b/Engine/Source/Runtime/LevelSequence/Public/LevelSequenceShotMetaDataLibrary.h
@@ -0,0 +1,109 @@
+// Copyright Epic Games, Inc. All Rights Reserved.
+
+#pragma once
+
+#include "AssetRegistry/AssetData.h"
+#include "Kismet/BlueprintFunctionLibrary.h"
+#include "LevelSequenceShotMetaDataLibrary.generated.h"
+
+class ULevelSequence;
+struct FAssetData;
+
+/**
+ * Manages ULevelSequence meta data that is common for rating shots in cinematic workflows.
+ * Manages access to UMovieSceneShotsMetaData for ULevelSequences.
+ */
+UCLASS()
+class LEVELSEQUENCE_API ULevelSequenceShotMetaDataLibrary : public UBlueprintFunctionLibrary
+{
+	GENERATED_BODY()
+public:
+	
+	/**
+	 * Gets whether the user has marked this shot as no good.
+	 * @return Whether the value was set by the user.
+	 */
+	UFUNCTION(BlueprintPure, Category = "Level Sequence")
+	static bool GetIsNoGood(const ULevelSequence* InLevelSequence, bool& bOutNoGood);
+	/**
+	 * Gets whether this was flagged by a user.
+	 * @return Whether the value was set by the user.
+	 */
+	UFUNCTION(BlueprintPure, Category = "Level Sequence")
+	static bool GetIsFlagged(const ULevelSequence* InLevelSequence, bool& bOutIsFlagged);
+	/**
+	 * Gets the favorite rating. The favorite rating is like a star rating, usually 1-3 if it was rated.
+	 * @return Whether the value was set by the user.
+	 */
+	UFUNCTION(BlueprintPure, Category = "Level Sequence")
+	static bool GetFavoriteRating(const ULevelSequence* InLevelSequence, int32& OutFavoriteRating);
+
+	/** @return Whether the user has ever set a value for this property. */
+	UFUNCTION(BlueprintPure, Category = "Level Sequence")
+	static bool HasIsNoGood(const ULevelSequence* InLevelSequence);
+	/** @return Whether the user has ever set a value for this property. */
+	UFUNCTION(BlueprintPure, Category = "Level Sequence")
+	static bool HasIsFlagged(const ULevelSequence* InLevelSequence);
+	/** @return Whether the user has ever set a value for this property. */
+	UFUNCTION(BlueprintPure, Category = "Level Sequence")
+	static bool HasFavoriteRating(const ULevelSequence* InLevelSequence);
+
+	/** Sets whether the user has marked this shot as no good. */
+	UFUNCTION(BlueprintCallable, Category = "Level Sequence")
+	static void SetIsNoGood(ULevelSequence* InLevelSequence, bool bInIsNoGood);
+	/** Sets whether this was flagged by a user. */
+	UFUNCTION(BlueprintCallable, Category = "Level Sequence")
+	static void SetIsFlagged(ULevelSequence* InLevelSequence, bool bInIsFlagged);
+	/** Sets the favorite rating. The favorite rating is like a star rating, usually 1-3 if it was rated. */
+	UFUNCTION(BlueprintCallable, Category = "Level Sequence")
+	static void SetFavoriteRating(ULevelSequence* InLevelSequence, int32 InFavoriteRating);
+
+	/** Clear the IsNoGood flag. */
+	UFUNCTION(BlueprintCallable, Category = "Level Sequence")
+	static void ClearIsNoGood(ULevelSequence* InLevelSequence);
+	/** Clear IsFlagged flag. */
+	UFUNCTION(BlueprintCallable, Category = "Level Sequence")
+	static void ClearIsFlagged(ULevelSequence* InLevelSequence);
+	/** Clears the favorite rating. */
+	UFUNCTION(BlueprintCallable, Category = "Level Sequence")
+	static void ClearFavoriteRating(ULevelSequence* InLevelSequence);
+
+	/** @return The asset tag for the IsNoGood value. */
+	UFUNCTION(BlueprintPure, Category = "Level Sequence")
+	static FName GetIsNoGoodAssetTag();
+	/** @return The asset tag for the IsFlagged value. */
+	UFUNCTION(BlueprintPure, Category = "Level Sequence")
+	static FName GetIsFlaggedAssetTag();
+	/** @return The asset tag for the FavoriteRating value. */
+	UFUNCTION(BlueprintPure, Category = "Level Sequence")
+	static FName GetFavoriteRatingAssetTag();
+
+	/**
+	 * Gets whether the user has marked this shot as no good.
+	 * @return Whether the value was set by the user.
+	 */
+	UFUNCTION(BlueprintPure, Category = "Level Sequence")
+	static bool GetIsNoGoodByAssetData(const FAssetData& InAssetData, bool& bOutNoGood);
+	/**
+	 * Gets whether this was flagged by a user.
+	 * @return Whether the value was set by the user.
+	 */
+	UFUNCTION(BlueprintPure, Category = "Level Sequence")
+	static bool GetIsFlaggedByAssetData(const FAssetData& InAssetData, bool& bOutIsFlagged);
+	/**
+	 * Gets the favorite rating. The favorite rating is like a star rating, usually 1-3 if it was rated.
+	 * @return Whether the value was set by the user.
+	 */
+	UFUNCTION(BlueprintPure, Category = "Level Sequence")
+	static bool GetFavoriteRatingByAssetData(const FAssetData& InAssetData, int32& OutFavoriteRating);
+	
+	/** @return Whether the user has ever set a value for this property. */
+	UFUNCTION(BlueprintPure, Category = "Level Sequence")
+	static bool HasIsNoGoodByAssetData(const FAssetData& InAssetData);
+	/** @return Whether the user has ever set a value for this property. */
+	UFUNCTION(BlueprintPure, Category = "Level Sequence")
+	static bool HasIsFlaggedByAssetData(const FAssetData& InAssetData);
+	/** @return Whether the user has ever set a value for this property. */
+	UFUNCTION(BlueprintPure, Category = "Level Sequence")
+	static bool HasFavoriteRatingByAssetData(const FAssetData& InAssetData);
+};
diff --git a/Engine/Source/Runtime/MovieScene/Private/MetaData/MovieSceneShotMetaData.cpp b/Engine/Source/Runtime/MovieScene/Private/MetaData/MovieSceneShotMetaData.cpp
new file mode 100644
--- /dev/null
+++ b/Engine/Source/Runtime/MovieScene/Private/MetaData/MovieSceneShotMetaData.cpp
@@ -0,0 +1,83 @@
+// Copyright Epic Games, Inc. All Rights Reserved.
+
+#include "MetaData/MovieSceneShotMetaData.h"
+
+#include "AssetRegistry/AssetData.h"
+#include "Containers/UnrealString.h"
+#include "Internationalization/Text.h"
+
+#define LOCTEXT_NAMESPACE "MovieSceneShotMetaData"
+
+namespace UE::MovieScene::ShotMetaDataDetail
+{
+template<typename TTagValue>
+static bool ParseAssetTag(const FAssetData& InAssetData, FName InTag, TTagValue& Value)
+{
+	FString StringValue;
+	InAssetData.GetTagValue(InTag, StringValue);
+	
+	// ExtendAssetRegistryTags sets "Unset" if the optional value was unset. Users could have also messed with the registry.
+	if (StringValue.IsNumeric()) 
+	{
+		// Ok to demote int to bool
+		Value = static_cast<TTagValue>(FCString::Atoi(*StringValue));
+		return true;
+	}
+	return false;
+}
+}
+
+const FName UMovieSceneShotMetaData::AssetRegistryTag_bIsNoGood(TEXT("AssetRegistryTag_MovieScene_bIsNoGood"));
+const FName UMovieSceneShotMetaData::AssetRegistryTag_bIsFlagged(TEXT("AssetRegistryTag_MovieScene_bIsFlagged"));
+const FName UMovieSceneShotMetaData::AssetRegistryTag_FavoriteRating(TEXT("AssetRegistryTag_MovieScene_FavoriteRating"));
+
+bool UMovieSceneShotMetaData::GetIsNoGoodByAssetData(const FAssetData& InAssetData, bool& bOutNoGood)
+{
+	return UE::MovieScene::ShotMetaDataDetail::ParseAssetTag(InAssetData, AssetRegistryTag_bIsNoGood, bOutNoGood);
+}
+
+bool UMovieSceneShotMetaData::GetIsFlaggedByAssetData(const FAssetData& InAssetData, bool& bOutIsFlagged)
+{
+	return UE::MovieScene::ShotMetaDataDetail::ParseAssetTag(InAssetData, AssetRegistryTag_bIsFlagged, bOutIsFlagged);
+}
+
+bool UMovieSceneShotMetaData::GetFavoriteRatingByAssetData(const FAssetData& InAssetData, int32& OutFavoriteRating)
+{
+	return UE::MovieScene::ShotMetaDataDetail::ParseAssetTag(InAssetData, AssetRegistryTag_FavoriteRating, OutFavoriteRating);
+}
+
+void UMovieSceneShotMetaData::ExtendAssetRegistryTags(FAssetRegistryTagsContext Context) const
+{
+	Context.AddTag(FAssetRegistryTag(
+		AssetRegistryTag_bIsNoGood, bIsNoGood.IsSet() ? FString::FromInt(*bIsNoGood) : TEXT("Unset"),
+		FAssetRegistryTag::ETagType::TT_Numerical, FAssetRegistryTag::TD_None
+	));
+	Context.AddTag(FAssetRegistryTag(
+		AssetRegistryTag_bIsFlagged, bIsFlagged.IsSet() ? FString::FromInt(*bIsFlagged) : TEXT("Unset"),
+		FAssetRegistryTag::ETagType::TT_Numerical, FAssetRegistryTag::TD_None
+	));
+	Context.AddTag(FAssetRegistryTag(
+		AssetRegistryTag_FavoriteRating, FavoriteRating.IsSet() ? FString::FromInt(*FavoriteRating) : TEXT("Unset"),
+		FAssetRegistryTag::ETagType::TT_Alphabetical, FAssetRegistryTag::TD_None
+	));
+}
+
+#if WITH_EDITOR
+void UMovieSceneShotMetaData::ExtendAssetRegistryTagMetaData(TMap<FName, UObject::FAssetRegistryTagMetadata>& OutMetadata) const
+{
+	OutMetadata.Add(AssetRegistryTag_bIsNoGood, FAssetRegistryTagMetadata()
+		.SetDisplayName(LOCTEXT("IsNoGood.Label", "Is No Good"))
+		.SetTooltip(LOCTEXT("IsNoGood.Description", "Whether this shot does not fulfill your requirements."))
+	);
+	OutMetadata.Add(AssetRegistryTag_bIsFlagged, FAssetRegistryTagMetadata()
+		.SetDisplayName(LOCTEXT("IsFlagged.Label", "Is Flagged"))
+		.SetTooltip(LOCTEXT("IsFlagged.Description", "Whether this shot was highlighted"))
+	);
+	OutMetadata.Add(AssetRegistryTag_FavoriteRating, FAssetRegistryTagMetadata()
+		.SetDisplayName(LOCTEXT("FavoriteRating.Label", "Rating"))
+		.SetTooltip(LOCTEXT("FavoriteRating.Description", "A star rating, usually ranging 1-3 and 0 meaning no rating, yet."))
+	);
+}
+#endif
+
+#undef LOCTEXT_NAMESPACE
\ No newline at end of file
diff --git a/Engine/Source/Runtime/MovieScene/Public/MetaData/MovieSceneShotMetaData.h b/Engine/Source/Runtime/MovieScene/Public/MetaData/MovieSceneShotMetaData.h
new file mode 100644
--- /dev/null
+++ b/Engine/Source/Runtime/MovieScene/Public/MetaData/MovieSceneShotMetaData.h
@@ -0,0 +1,71 @@
+// Copyright Epic Games, Inc. All Rights Reserved.
+
+#pragma once
+
+#include "IMovieSceneMetaData.h"
+#include "UObject/Object.h"
+#include "MovieSceneShotMetaData.generated.h"
+
+struct FAssetData;
+
+/**
+ * Holds meta data that is common for rating shots in cinematic workflows.
+ * The purpose is to share this data among related tools.
+ *
+ * Stored in ULevelSequence as meta  data. 
+ * @see ULevelSequenceShotMetaDataLibrary for managing data.
+ */
+UCLASS()
+class MOVIESCENE_API UMovieSceneShotMetaData
+	: public UObject
+	, public IMovieSceneMetaDataInterface
+{
+	GENERATED_BODY()
+public:
+
+	/** The asset registry tag that contains whether this shot is good or not */
+	static const FName AssetRegistryTag_bIsNoGood;
+	/** The asset registry tag that contains whether this was flagged by a user */
+	static const FName AssetRegistryTag_bIsFlagged;
+	/** The asset registry tag that contains its favorite status */
+	static const FName AssetRegistryTag_FavoriteRating;
+
+	static TArray<FName> GetAllTags() { return { AssetRegistryTag_bIsNoGood, AssetRegistryTag_bIsFlagged, AssetRegistryTag_FavoriteRating }; }
+
+	const TOptional<bool>& GetIsNoGood() const { return bIsNoGood; }
+	const TOptional<bool>& GetIsFlagged() const { return bIsFlagged; }
+	const TOptional<int32>& GetFavoriteRating() const { return FavoriteRating; }
+
+	void SetIsNoGood(bool bInIsNoGood) { bIsNoGood = bInIsNoGood; }
+	void SetIsFlagged(bool bInIsFlagged) { bIsFlagged = bInIsFlagged; }
+	void SetFavoriteRating(int32 InFavoriteRating) { FavoriteRating = InFavoriteRating; }
+	
+	void ClearIsNoGood() { bIsNoGood.Reset(); }
+	void ClearIsFlagged() { bIsFlagged.Reset(); }
+	void ClearFavoriteRating() { FavoriteRating.Reset(); }
+
+	static bool GetIsNoGoodByAssetData(const FAssetData& InAssetData, bool& bOutNoGood);
+	static bool GetIsFlaggedByAssetData(const FAssetData& InAssetData, bool& bOutIsFlagged);
+	static bool GetFavoriteRatingByAssetData(const FAssetData& InAssetData, int32& OutFavoriteRating); 
+
+	// Begin IMovieSceneMetaDataInterface Interface
+	virtual void ExtendAssetRegistryTags(FAssetRegistryTagsContext Context) const override;
+#if WITH_EDITOR
+	virtual void ExtendAssetRegistryTagMetaData(TMap<FName, UObject::FAssetRegistryTagMetadata>& OutMetadata) const override;
+#endif
+	// End IMovieSceneMetaDataInterface Interface
+	
+private:
+
+	/** Whether this shot is marked as not good */
+	UPROPERTY()
+	TOptional<bool> bIsNoGood;
+	
+	/** The asset registry tag that contains whether this was flagged by a user */
+	UPROPERTY()
+	TOptional<bool> bIsFlagged;
+
+	/** The favorite rating is like a star rating, usually 1-3 if it was rated. */
+	UPROPERTY()
+	TOptional<int32> FavoriteRating;
+};
