Virtual Camera: Match missing PRAGMA_ENABLE_DEPRECATION_WARNINGS.

diff --git a/Engine/Plugins/VirtualProduction/VirtualCamera/Source/VirtualCamera/Private/LevelSequence/VirtualCameraClipsMetaData.cpp b/Engine/Plugins/VirtualProduction/VirtualCamera/Source/VirtualCamera/Private/LevelSequence/VirtualCameraClipsMetaData.cpp
--- a/Engine/Plugins/VirtualProduction/VirtualCamera/Source/VirtualCamera/Private/LevelSequence/VirtualCameraClipsMetaData.cpp
+++ b/Engine/Plugins/VirtualProduction/VirtualCamera/Source/VirtualCamera/Private/LevelSequence/VirtualCameraClipsMetaData.cpp
@@ -30,6 +30,7 @@ FName UVirtualCameraClipsMetaData::GetClipsMetaDataTag_FocalLength()
 {
 	return AssetRegistryTag_FocalLength;
 }
+PRAGMA_ENABLE_DEPRECATION_WARNINGS
 
 PRAGMA_DISABLE_DEPRECATION_WARNINGS
 FName UVirtualCameraClipsMetaData::GetClipsMetaDataTag_IsSelected()
