VCam: Enable VirtualCamera.AutoMigrateLevelSequenceOnAccess CVar by default. This was supposed to be enabled by default in the first place.

diff --git a/Engine/Plugins/VirtualProduction/VirtualCamera/Source/VirtualCamera/Private/LevelSequence/VCamTakesMetaDataMigration.cpp b/Engine/Plugins/VirtualProduction/VirtualCamera/Source/VirtualCamera/Private/LevelSequence/VCamTakesMetaDataMigration.cpp
--- a/Engine/Plugins/VirtualProduction/VirtualCamera/Source/VirtualCamera/Private/LevelSequence/VCamTakesMetaDataMigration.cpp
+++ b/Engine/Plugins/VirtualProduction/VirtualCamera/Source/VirtualCamera/Private/LevelSequence/VCamTakesMetaDataMigration.cpp
@@ -16,7 +16,7 @@ namespace UE::VirtualCamera::MigrationDetail
 {
 static TAutoConsoleVariable<bool> CVarAutoMigrateLevelSequenceOnAccess(
 	TEXT("VirtualCamera.AutoMigrateLevelSequenceOnAccess"),
-	false,
+	true,
 	TEXT("Whenever VirtualCamera opens a level sequence, whether to automatically migrate the meta data from the old UVirtualCameraClipsMetaData to the new UMovieSceneShotMetaData.")
 );
 
