Virtual Camera: Fix null SourceObject in SupportsBindingCreationFromObject

#rnx

diff --git a/Engine/Plugins/VirtualProduction/VirtualCamera/Source/VirtualCamera/Private/LevelSequence/MovieSceneSpawnableVCamBinding.cpp b/Engine/Plugins/VirtualProduction/VirtualCamera/Source/VirtualCamera/Private/LevelSequence/MovieSceneSpawnableVCamBinding.cpp
--- a/Engine/Plugins/VirtualProduction/VirtualCamera/Source/VirtualCamera/Private/LevelSequence/MovieSceneSpawnableVCamBinding.cpp
+++ b/Engine/Plugins/VirtualProduction/VirtualCamera/Source/VirtualCamera/Private/LevelSequence/MovieSceneSpawnableVCamBinding.cpp
@@ -43,6 +43,11 @@ bool UMovieSceneSpawnableVCamBinding::CanRecordComponent(const UE::TakeRecorderS
 
 bool UMovieSceneSpawnableVCamBinding::SupportsBindingCreationFromObject(const UObject* SourceObject) const
 {
+	if (!SourceObject)
+	{
+		return false;
+	}
+
 	return UE::VirtualCamera::FVCamInfo(SourceObject).ShouldRecordAsCineCamera();
 }
 
