Take Recorder x Virtual Camera:
- The Virtual Camera plugin can now be disabled after recording: UMovieSceneSpawnableVCamBinding now creates a UMovieSceneSpawnableActorBinding instead (since the VCamBinding class would no longer exist if  the plugin is disabled).
- Fix recorded Virtual Camera having a CineCaptureComponent2D recorded
- Fixes post process settings of owning camera not getting recorded correctly.

diff --git a/Engine/Plugins/VirtualProduction/VirtualCamera/Source/VirtualCamera/Private/LevelSequence/MovieSceneSpawnableVCamBinding.cpp b/Engine/Plugins/VirtualProduction/VirtualCamera/Source/VirtualCamera/Private/LevelSequence/MovieSceneSpawnableVCamBinding.cpp
--- a/Engine/Plugins/VirtualProduction/VirtualCamera/Source/VirtualCamera/Private/LevelSequence/MovieSceneSpawnableVCamBinding.cpp
+++ b/Engine/Plugins/VirtualProduction/VirtualCamera/Source/VirtualCamera/Private/LevelSequence/MovieSceneSpawnableVCamBinding.cpp
@@ -3,7 +3,6 @@
 #include "LevelSequence/MovieSceneSpawnableVCamBinding.h"
 
 #include "CineCameraActor.h"
-#include "Components/InputComponent.h"
 #include "Engine/Engine.h"
 #include "MovieScene.h"
 #include "VCamComponent.h"
@@ -18,11 +17,14 @@ struct FVCamInfo
 {
 	const AActor* Actor;
 	UVCamComponent* VCamComponent;
+
+	USceneComponent* RootComponent;
 	UCineCameraComponent* Camera;
 
 	explicit FVCamInfo(TNotNull<const UObject*> SourceObject)
 		: Actor(SourceObject->IsA<AActor>() ? Cast<AActor>(SourceObject) : SourceObject->GetTypedOuter<AActor>())
 		, VCamComponent(Actor ? Actor->FindComponentByClass<UVCamComponent>() : nullptr)
+		, RootComponent(Actor ? Actor->GetRootComponent() : nullptr)
 		, Camera(VCamComponent ? VCamComponent->GetTargetCamera() : nullptr)
 	{}
 
@@ -35,20 +37,19 @@ bool UMovieSceneSpawnableVCamBinding::CanRecordComponent(const UE::TakeRecorderS
 {
 	// If recording as cine camera, prevent the following components
 	const UE::VirtualCamera::FVCamInfo VCamInfo(InArgs.ObjectToRecord);
-	const bool bIsVCamComponent = VCamInfo.ShouldRecordAsCineCamera() && InArgs.ObjectToRecord == VCamInfo.VCamComponent;
-	const bool bIsInputComponent = VCamInfo.ShouldRecordAsCineCamera() && InArgs.ObjectToRecord == VCamInfo.VCamComponent->GetInputComponent();
-	return !bIsVCamComponent && !bIsInputComponent;
+	
+	// Only record components that have equivalents on  ACineCameraActor, i.e. skip VCamComponent, InputComponent, CineCaptureComponent2D, etc.
+	const bool bIsComponentFromCineCameraActor = InArgs.ObjectToRecord == VCamInfo.RootComponent || InArgs.ObjectToRecord == VCamInfo.Camera;
+	// If VCamActor is not configured to record as cine camera, allow all components.
+	const bool bAllowAllComponents = !VCamInfo.ShouldRecordAsCineCamera();
+	
+	return bAllowAllComponents || bIsComponentFromCineCameraActor || InArgs.ObjectToRecord == VCamInfo.Actor;
 }
 #endif
 
 bool UMovieSceneSpawnableVCamBinding::SupportsBindingCreationFromObject(const UObject* SourceObject) const
 {
-	if (!SourceObject)
-	{
-		return false;
-	}
-
-	return UE::VirtualCamera::FVCamInfo(SourceObject).ShouldRecordAsCineCamera();
+	return SourceObject && UE::VirtualCamera::FVCamInfo(SourceObject).ShouldRecordAsCineCamera();
 }
 
 UMovieSceneCustomBinding* UMovieSceneSpawnableVCamBinding::CreateNewCustomBinding(UObject* SourceObject, UMovieScene& OwnerMovieScene)
@@ -80,15 +81,12 @@ UMovieSceneCustomBinding* UMovieSceneSpawnableVCamBinding::CreateNewCustomBindin
 #if WITH_EDITORONLY_DATA
 	CineCameraActor->bIsEditorPreviewActor = false;
 #endif
-	
-	UMovieSceneSpawnableVCamBinding* Binding = NewObject<UMovieSceneSpawnableVCamBinding>(&OwnerMovieScene, GetClass(), InstancedBindingName, RF_Transactional);
+
+	// Record as if this was CineCameraActor.
+	// Do not create UMovieSceneSpawnableVCamBinding instance because that class won't be available if the user disables the Virtual Camera plugin.
+	UMovieSceneSpawnableActorBinding* Binding = NewObject<UMovieSceneSpawnableActorBinding>(
+		&OwnerMovieScene, UMovieSceneSpawnableActorBinding::StaticClass(), InstancedBindingName, RF_Transactional
+		);
 	Binding->SetObjectTemplate(CineCameraActor);
 	return Binding;
-}
-
-#if WITH_EDITOR
-FText UMovieSceneSpawnableVCamBinding::GetBindingTypePrettyName() const
-{
-	return NSLOCTEXT("VirtualCamera", "MovieSceneSpawnableVCamBinding", "Spawnable Virtual Camera");
-}
-#endif
\ No newline at end of file
+}
\ No newline at end of file
diff --git a/Engine/Plugins/VirtualProduction/VirtualCamera/Source/VirtualCamera/Public/LevelSequence/MovieSceneSpawnableVCamBinding.h b/Engine/Plugins/VirtualProduction/VirtualCamera/Source/VirtualCamera/Public/LevelSequence/MovieSceneSpawnableVCamBinding.h
--- a/Engine/Plugins/VirtualProduction/VirtualCamera/Source/VirtualCamera/Public/LevelSequence/MovieSceneSpawnableVCamBinding.h
+++ b/Engine/Plugins/VirtualProduction/VirtualCamera/Source/VirtualCamera/Public/LevelSequence/MovieSceneSpawnableVCamBinding.h
@@ -27,9 +27,6 @@ public:
 	//~ Begin UMovieSceneSpawnableActorBinding Interface
 	virtual bool SupportsBindingCreationFromObject(const UObject* SourceObject) const override;
 	virtual UMovieSceneCustomBinding* CreateNewCustomBinding(UObject* SourceObject, UMovieScene& OwnerMovieScene) override;
-#if WITH_EDITOR
-	virtual FText GetBindingTypePrettyName() const override;
-#endif
 	//~ End UMovieSceneSpawnableActorBinding Interface
 
 protected:
