Take Recorder x VCam: Make template actor transactional and skip the rest of flags. Other flags, such as RF_ExternalPackage, would get carried over and cause AActor::IsAsset to return true, causing failing ensures.

diff --git a/Engine/Plugins/VirtualProduction/VirtualCamera/Source/VirtualCamera/Private/LevelSequence/MovieSceneSpawnableVCamBinding.cpp b/Engine/Plugins/VirtualProduction/VirtualCamera/Source/VirtualCamera/Private/LevelSequence/MovieSceneSpawnableVCamBinding.cpp
--- a/Engine/Plugins/VirtualProduction/VirtualCamera/Source/VirtualCamera/Private/LevelSequence/MovieSceneSpawnableVCamBinding.cpp
+++ b/Engine/Plugins/VirtualProduction/VirtualCamera/Source/VirtualCamera/Private/LevelSequence/MovieSceneSpawnableVCamBinding.cpp
@@ -63,7 +63,7 @@ UMovieSceneCustomBinding* UMovieSceneSpawnableVCamBinding::CreateNewCustomBindin
 	const FName TemplateName = MakeUniqueObjectName(&OwnerMovieScene, UObject::StaticClass(), SourceObject ? SourceObject->GetFName() : TEXT("EmptyBinding"));
 	const FName InstancedBindingName = MakeUniqueObjectName(&OwnerMovieScene, UObject::StaticClass(), *FString(TemplateName.ToString() + TEXT("_CustomBinding")));
 
-	ACineCameraActor* CineCameraActor = NewObject<ACineCameraActor>(&OwnerMovieScene, TemplateName, SourceObject->GetFlags() & ~RF_Transient | RF_Transactional | RF_Public);
+	ACineCameraActor* CineCameraActor = NewObject<ACineCameraActor>(&OwnerMovieScene, TemplateName, RF_Transactional);
 	if (!CineCameraActor)
 	{
 		return nullptr;
