VCam: Refactor interaction with Take Recorder / Sequencer to FVCamSupportForCinematicTooling. In the future, any further customization goes here.

diff --git a/Engine/Plugins/VirtualProduction/VirtualCamera/Source/VirtualCamera/Private/LevelSequence/MovieSceneSpawnableVCamBinding.cpp b/Engine/Plugins/VirtualProduction/VirtualCamera/Source/VirtualCamera/Private/LevelSequence/MovieSceneSpawnableVCamBinding.cpp
deleted file mode 100644
diff --git a/Engine/Plugins/VirtualProduction/VirtualCamera/Source/VirtualCamera/Private/VirtualCamera.cpp b/Engine/Plugins/VirtualProduction/VirtualCamera/Source/VirtualCamera/Private/VirtualCamera.cpp
--- a/Engine/Plugins/VirtualProduction/VirtualCamera/Source/VirtualCamera/Private/VirtualCamera.cpp
+++ b/Engine/Plugins/VirtualProduction/VirtualCamera/Source/VirtualCamera/Private/VirtualCamera.cpp
@@ -3,7 +3,6 @@
 #include "VirtualCamera.h"
 
 #include "AdvancedWidgetsModule.h"
-#include "LevelSequence/MovieSceneSpawnableVCamBinding.h"
 #include "Misc/CoreDelegates.h"
 
 #if WITH_EDITOR
@@ -28,29 +27,11 @@ namespace UE::VirtualCamera
 	void FVirtualCameraModuleImpl::ShutdownModule()
 	{
 		UnregisterSettings();
-
-#if WITH_EDITOR
-		using namespace TakeRecorderSources;
-		if (ITakeRecorderSourcesModule::IsAvailable())
-		{
-			ITakeRecorderSourcesModule::Get().UnregisterCanRecordDelegate(TEXT("VirtualCamera_CanRecordDelegate"));
-		}
-#endif
 	}
 
 	void FVirtualCameraModuleImpl::OnPostInitEngine()
 	{
 		RegisterSettings();
-
-#if WITH_EDITOR
-		using namespace TakeRecorderSources;
-		if (ITakeRecorderSourcesModule::IsAvailable())
-		{
-			ITakeRecorderSourcesModule::Get().RegisterCanRecordDelegate(
-				TEXT("VirtualCamera_CanRecordDelegate"), FCanRecordDelegate::CreateStatic(&UMovieSceneSpawnableVCamBinding::CanRecordComponent)
-				);
-		}
-#endif
 	}
 
 	void FVirtualCameraModuleImpl::RegisterSettings()
diff --git a/Engine/Plugins/VirtualProduction/VirtualCamera/Source/VirtualCamera/Public/LevelSequence/MovieSceneSpawnableVCamBinding.h b/Engine/Plugins/VirtualProduction/VirtualCamera/Source/VirtualCamera/Public/LevelSequence/MovieSceneSpawnableVCamBinding.h
deleted file mode 100644
diff --git a/Engine/Plugins/VirtualProduction/VirtualCamera/Source/VirtualCamera/VirtualCamera.Build.cs b/Engine/Plugins/VirtualProduction/VirtualCamera/Source/VirtualCamera/VirtualCamera.Build.cs
--- a/Engine/Plugins/VirtualProduction/VirtualCamera/Source/VirtualCamera/VirtualCamera.Build.cs
+++ b/Engine/Plugins/VirtualProduction/VirtualCamera/Source/VirtualCamera/VirtualCamera.Build.cs
@@ -20,7 +20,6 @@ public class VirtualCamera : ModuleRules
 				"LevelSequence",
 				"LiveLinkInterface",
 				"MovieScene",
-				"MovieSceneTracks",
 				"RemoteSession",
 				"TimeManagement",
 				"UMG",
diff --git a/Engine/Plugins/VirtualProduction/VirtualCamera/Source/VirtualCameraEditor/Private/Cinematic/MovieSceneSpawnableVCamBinding.cpp b/Engine/Plugins/VirtualProduction/VirtualCamera/Source/VirtualCameraEditor/Private/Cinematic/MovieSceneSpawnableVCamBinding.cpp
new file mode 100644
--- /dev/null
+++ b/Engine/Plugins/VirtualProduction/VirtualCamera/Source/VirtualCameraEditor/Private/Cinematic/MovieSceneSpawnableVCamBinding.cpp
@@ -0,0 +1,50 @@
+﻿// Copyright Epic Games, Inc. All Rights Reserved.
+
+#include "MovieSceneSpawnableVCamBinding.h"
+
+#include "CineCameraActor.h"
+#include "Engine/Engine.h"
+#include "MovieScene.h"
+#include "VCamHierarchyInfo.h"
+
+bool UMovieSceneSpawnableVCamBinding::SupportsBindingCreationFromObject(const UObject* SourceObject) const
+{
+	return SourceObject && UE::VirtualCamera::FVCamHierarchyInfo(SourceObject).ShouldRecordAsCineCamera();
+}
+
+UMovieSceneCustomBinding* UMovieSceneSpawnableVCamBinding::CreateNewCustomBinding(UObject* SourceObject, UMovieScene& OwnerMovieScene)
+{
+	const UE::VirtualCamera::FVCamHierarchyInfo VCamInfo(SourceObject);
+	if (!SourceObject || !VCamInfo.ShouldRecordAsCineCamera()) // We don't expect this to happen as SupportsBindingCreationFromObject rejects this case.
+	{
+		return nullptr;
+	}
+	
+	const FName TemplateName = MakeUniqueObjectName(&OwnerMovieScene, UObject::StaticClass(), SourceObject ? SourceObject->GetFName() : TEXT("EmptyBinding"));
+	const FName InstancedBindingName = MakeUniqueObjectName(&OwnerMovieScene, UObject::StaticClass(), *FString(TemplateName.ToString() + TEXT("_CustomBinding")));
+
+	ACineCameraActor* CineCameraActor = NewObject<ACineCameraActor>(&OwnerMovieScene, TemplateName, RF_Transactional);
+	if (!CineCameraActor)
+	{
+		return nullptr;
+	}
+	
+	UEngine::FCopyPropertiesForUnrelatedObjectsParams CopyParams;
+	CopyParams.bNotifyObjectReplacement = false;
+	CopyParams.bPreserveRootComponent = false;
+	CopyParams.bPerformDuplication = true;
+	UEngine::CopyPropertiesForUnrelatedObjects(SourceObject, CineCameraActor, CopyParams);
+	
+	CineCameraActor->Tags.Remove(TEXT("SequencerActor"));
+	CineCameraActor->Tags.Remove(TEXT("SequencerPreviewActor"));
+	CineCameraActor->DetachFromActor(FDetachmentTransformRules(EDetachmentRule::KeepRelative, false));
+	CineCameraActor->bIsEditorPreviewActor = false;
+
+	// Record as if this was CineCameraActor.
+	// Do not create UMovieSceneSpawnableVCamBinding instance because that class won't be available if the user disables the Virtual Camera plugin.
+	UMovieSceneSpawnableActorBinding* Binding = NewObject<UMovieSceneSpawnableActorBinding>(
+		&OwnerMovieScene, UMovieSceneSpawnableActorBinding::StaticClass(), InstancedBindingName, RF_Transactional
+		);
+	Binding->SetObjectTemplate(CineCameraActor);
+	return Binding;
+}
\ No newline at end of file
diff --git a/Engine/Plugins/VirtualProduction/VirtualCamera/Source/VirtualCameraEditor/Private/Cinematic/MovieSceneSpawnableVCamBinding.h b/Engine/Plugins/VirtualProduction/VirtualCamera/Source/VirtualCameraEditor/Private/Cinematic/MovieSceneSpawnableVCamBinding.h
new file mode 100644
--- /dev/null
+++ b/Engine/Plugins/VirtualProduction/VirtualCamera/Source/VirtualCameraEditor/Private/Cinematic/MovieSceneSpawnableVCamBinding.h
@@ -0,0 +1,36 @@
+﻿// Copyright Epic Games, Inc. All Rights Reserved.
+
+#pragma once
+
+#include "Bindings/MovieSceneSpawnableActorBinding.h"
+#include "MovieSceneSpawnableVCamBinding.generated.h"
+
+namespace UE::TakeRecorderSources{ struct FCanRecordArgs; }
+
+/**
+ * Allows the AVCamBaseActor to be recorded as ACineCameraActor instead.
+ *
+ * This allows users to record a VCam actor using Take Recorder and disable the VCam plugin afterwards while still being able to use the recorded
+ * sequences without missing classes.
+ *
+ * This class is never actually instantiated. It simply creates a UMovieSceneSpawnableActorBinding that contains a ACineCameraActor instead of a VCam.
+ * Not instantiating this class is important: otherwise the sequence referencing the binding would lose the reference when VCam plugin is disabled as
+ * the class would be missing.
+ */
+UCLASS()
+class UMovieSceneSpawnableVCamBinding : public UMovieSceneSpawnableActorBinding
+{
+	GENERATED_BODY()
+public:
+
+	//~ Begin UMovieSceneSpawnableActorBinding Interface
+	virtual bool SupportsBindingCreationFromObject(const UObject* SourceObject) const override;
+	virtual UMovieSceneCustomBinding* CreateNewCustomBinding(UObject* SourceObject, UMovieScene& OwnerMovieScene) override;
+	//~ End UMovieSceneSpawnableActorBinding Interface
+
+protected:
+	
+	//~ Begin UMovieSceneSpawnableActorBinding Interface
+	virtual int32 GetCustomBindingPriority() const override { return BaseCustomPriority; }
+	//~ End UMovieSceneSpawnableActorBinding Interface
+};
\ No newline at end of file
diff --git a/Engine/Plugins/VirtualProduction/VirtualCamera/Source/VirtualCameraEditor/Private/Cinematic/VCamHierarchyInfo.h b/Engine/Plugins/VirtualProduction/VirtualCamera/Source/VirtualCameraEditor/Private/Cinematic/VCamHierarchyInfo.h
new file mode 100644
--- /dev/null
+++ b/Engine/Plugins/VirtualProduction/VirtualCamera/Source/VirtualCameraEditor/Private/Cinematic/VCamHierarchyInfo.h
@@ -0,0 +1,30 @@
+﻿// Copyright Epic Games, Inc. All Rights Reserved.
+
+#pragma once
+
+#include "CineCameraComponent.h"
+#include "GameFramework/Actor.h"
+#include "VCamComponent.h"
+#include "Misc/NotNull.h"
+
+namespace UE::VirtualCamera
+{
+/** Gets components of a VCamActor for the purpose of deciding whether it should be recorded as cine camera in Sequencer. */
+struct FVCamHierarchyInfo
+{
+	const AActor* Actor;
+	UVCamComponent* VCamComponent;
+
+	USceneComponent* RootComponent;
+	UCineCameraComponent* Camera;
+
+	explicit FVCamHierarchyInfo(TNotNull<const UObject*> SourceObject)
+		: Actor(SourceObject->IsA<AActor>() ? Cast<AActor>(SourceObject) : SourceObject->GetTypedOuter<AActor>())
+		, VCamComponent(Actor ? Actor->FindComponentByClass<UVCamComponent>() : nullptr)
+		, RootComponent(Actor ? Actor->GetRootComponent() : nullptr)
+		, Camera(VCamComponent ? VCamComponent->GetTargetCamera() : nullptr)
+	{}
+
+	bool ShouldRecordAsCineCamera() const { return VCamComponent != nullptr && Camera != nullptr && VCamComponent->GetRecordAsCineCamera(); }
+};
+}
\ No newline at end of file
diff --git a/Engine/Plugins/VirtualProduction/VirtualCamera/Source/VirtualCameraEditor/Private/Cinematic/VCamSupportForCinematicTooling.cpp b/Engine/Plugins/VirtualProduction/VirtualCamera/Source/VirtualCameraEditor/Private/Cinematic/VCamSupportForCinematicTooling.cpp
new file mode 100644
--- /dev/null
+++ b/Engine/Plugins/VirtualProduction/VirtualCamera/Source/VirtualCameraEditor/Private/Cinematic/VCamSupportForCinematicTooling.cpp
@@ -0,0 +1,42 @@
+﻿// Copyright Epic Games, Inc. All Rights Reserved.
+
+#include "VCamSupportForCinematicTooling.h"
+
+#include "ITakeRecorderSourcesModule.h"
+#include "VCamHierarchyInfo.h"
+
+namespace UE::VirtualCamera
+{
+FVCamSupportForCinematicTooling::FVCamSupportForCinematicTooling()
+{
+	using namespace TakeRecorderSources;
+	if (ITakeRecorderSourcesModule::IsAvailable())
+	{
+		ITakeRecorderSourcesModule::Get().RegisterCanRecordDelegate(
+			TEXT("VirtualCamera_CanRecordDelegate"), FCanRecordDelegate::CreateStatic(&FVCamSupportForCinematicTooling::CanRecordComponent)
+			);
+	}
+}
+
+FVCamSupportForCinematicTooling::~FVCamSupportForCinematicTooling()
+{
+	using namespace TakeRecorderSources;
+	if (ITakeRecorderSourcesModule::IsAvailable())
+	{
+		ITakeRecorderSourcesModule::Get().UnregisterCanRecordDelegate(TEXT("VirtualCamera_CanRecordDelegate"));
+	}
+}
+
+bool FVCamSupportForCinematicTooling::CanRecordComponent(const TakeRecorderSources::FCanRecordArgs& InArgs)
+{
+	// If recording as cine camera, prevent the following components
+	const FVCamHierarchyInfo VCamInfo(InArgs.ObjectToRecord);
+	
+	// Only record components that have equivalents on  ACineCameraActor, i.e. skip VCamComponent, InputComponent, CineCaptureComponent2D, etc.
+	const bool bIsComponentFromCineCameraActor = InArgs.ObjectToRecord == VCamInfo.RootComponent || InArgs.ObjectToRecord == VCamInfo.Camera;
+	// If VCamActor is not configured to record as cine camera, allow all components.
+	const bool bAllowAllComponents = !VCamInfo.ShouldRecordAsCineCamera();
+	
+	return bAllowAllComponents || bIsComponentFromCineCameraActor || InArgs.ObjectToRecord == VCamInfo.Actor;
+}
+}
diff --git a/Engine/Plugins/VirtualProduction/VirtualCamera/Source/VirtualCameraEditor/Private/Cinematic/VCamSupportForCinematicTooling.h b/Engine/Plugins/VirtualProduction/VirtualCamera/Source/VirtualCameraEditor/Private/Cinematic/VCamSupportForCinematicTooling.h
new file mode 100644
--- /dev/null
+++ b/Engine/Plugins/VirtualProduction/VirtualCamera/Source/VirtualCameraEditor/Private/Cinematic/VCamSupportForCinematicTooling.h
@@ -0,0 +1,23 @@
+﻿// Copyright Epic Games, Inc. All Rights Reserved.
+
+#pragma once
+
+#include "Templates/UnrealTemplate.h"
+
+namespace UE::TakeRecorderSources { struct FCanRecordArgs; }
+
+namespace UE::VirtualCamera
+{
+/** Manages setting up global delegates, etc. so VCam operates correctly with Take Recorder & Sequencer. */
+class FVCamSupportForCinematicTooling : public FNoncopyable
+{
+public:
+
+	FVCamSupportForCinematicTooling();
+	~FVCamSupportForCinematicTooling();
+
+	/** Used by ITakeRecorderSourcesModule::RegisterCanRecordDelegate. If a VCam is set to record as ACineCameraActor, this will skip any additional components, such as VCam and the input component. */
+	static bool CanRecordComponent(const TakeRecorderSources::FCanRecordArgs& InArgs);
+};
+}
+
diff --git a/Engine/Plugins/VirtualProduction/VirtualCamera/Source/VirtualCameraEditor/Private/VirtualCameraEditorModule.cpp b/Engine/Plugins/VirtualProduction/VirtualCamera/Source/VirtualCameraEditor/Private/VirtualCameraEditorModule.cpp
--- a/Engine/Plugins/VirtualProduction/VirtualCamera/Source/VirtualCameraEditor/Private/VirtualCameraEditorModule.cpp
+++ b/Engine/Plugins/VirtualProduction/VirtualCamera/Source/VirtualCameraEditor/Private/VirtualCameraEditorModule.cpp
@@ -9,12 +9,16 @@
 #include "IVPUtilitiesEditorModule.h"
 #include "LevelEditor.h"
 #include "LevelEditorOutlinerSettings.h"
+#include "Cinematic/VCamSupportForCinematicTooling.h"
 #include "UObject/UObjectGlobals.h"
 
 #define LOCTEXT_NAMESPACE "FVirtualCameraEditorModule"
 
 DEFINE_LOG_CATEGORY_STATIC(LogVirtualCameraEditor, Log, All);
 
+namespace UE::VirtualCamera
+{
+
 class FVirtualCameraEditorModule : public IModuleInterface
 {
 public:
@@ -22,13 +26,20 @@ public:
 	{
 		RegisterPlacementModeItems();
 		RegisterOutlinerFilters();
+
+		CinematicToolingSupport = MakeUnique<FVCamSupportForCinematicTooling>();
 	}
 
 	virtual void ShutdownModule() override
-	{}
+	{
+		CinematicToolingSupport.Reset();
+	}
 
 private:
 	
+	/** Manages setting up global delegates, etc. so VCam operates correctly with Take Recorder & Sequencer. */
+	TUniquePtr<FVCamSupportForCinematicTooling> CinematicToolingSupport;
+	
 	void RegisterPlacementModeItems()
 	{
 		if (const FPlacementCategoryInfo* Info = IVPUtilitiesEditorModule::Get().GetVirtualProductionPlacementCategoryInfo()
@@ -79,8 +90,8 @@ private:
 		}
 	}
 };
-#undef LOCTEXT_NAMESPACE
 
 IMPLEMENT_MODULE(FVirtualCameraEditorModule, VirtualCameraEditor)
+}
 
-
+#undef LOCTEXT_NAMESPACE
\ No newline at end of file
diff --git a/Engine/Plugins/VirtualProduction/VirtualCamera/Source/VirtualCameraEditor/VirtualCameraEditor.Build.cs b/Engine/Plugins/VirtualProduction/VirtualCamera/Source/VirtualCameraEditor/VirtualCameraEditor.Build.cs
--- a/Engine/Plugins/VirtualProduction/VirtualCamera/Source/VirtualCameraEditor/VirtualCameraEditor.Build.cs
+++ b/Engine/Plugins/VirtualProduction/VirtualCamera/Source/VirtualCameraEditor/VirtualCameraEditor.Build.cs
@@ -16,9 +16,13 @@ public class VirtualCameraEditor : ModuleRules
 				"CinematicCamera",
 				"EditorWidgets",
 				"Engine",
+				"MovieScene",
+				"MovieSceneTracks",
 				"PlacementMode", 
 				"Settings",
 				"UnrealEd",
+				"TakeRecorderSources",
+				"VCamCore",
 				"VirtualCamera", 
 				"VPUtilitiesEditor", 
 			}
