Take Recorder x VCam: Reset UCameraComponent::AsymmetricOverscan in the track's template actor. The result is that the recorded cine camera only sees the inner, uncropped VCam viewport after recording. Before, it would see the entire viewport incluiding the bits that are supposed to be cropped.

diff --git a/Engine/Plugins/VirtualProduction/VirtualCamera/Source/VirtualCameraEditor/Private/Cinematic/MovieSceneSpawnableVCamBinding.cpp b/Engine/Plugins/VirtualProduction/VirtualCamera/Source/VirtualCameraEditor/Private/Cinematic/MovieSceneSpawnableVCamBinding.cpp
--- a/Engine/Plugins/VirtualProduction/VirtualCamera/Source/VirtualCameraEditor/Private/Cinematic/MovieSceneSpawnableVCamBinding.cpp
+++ b/Engine/Plugins/VirtualProduction/VirtualCamera/Source/VirtualCameraEditor/Private/Cinematic/MovieSceneSpawnableVCamBinding.cpp
@@ -34,6 +34,12 @@ UMovieSceneCustomBinding* UMovieSceneSpawnableVCamBinding::CreateNewCustomBindin
 	CopyParams.bPreserveRootComponent = false;
 	CopyParams.bPerformDuplication = true;
 	UEngine::CopyPropertiesForUnrelatedObjects(SourceObject, CineCameraActor, CopyParams);
+
+	// Specific for VCam: The asymmetric overscan needs to be zero in the resulting level sequence.
+	// The SensorCorrectionModifier sets the asymmetric overscan to be non-zero to achieve two "viewports" in the VCam UI.
+	//  - In inner, uncropped viewport which is supposed to be what the CineCameraActor is supposed to see after recording is done.
+	//  - The outer, cropped viewport, which is darker in the VCam UI. That section would be visible if we didn't reset asymmetric overscan here.
+	CineCameraActor->GetCineCameraComponent()->AsymmetricOverscan = FVector4f::Zero();
 	
 	CineCameraActor->Tags.Remove(TEXT("SequencerActor"));
 	CineCameraActor->Tags.Remove(TEXT("SequencerPreviewActor"));
