Take Recorder x VCam: Recorded CineCameraActor now adds black bars to side (bConstrainAspectRation = true) when playing back.

diff --git a/Engine/Plugins/VirtualProduction/VirtualCamera/Source/VirtualCameraEditor/Private/Cinematic/MovieSceneSpawnableVCamBinding.cpp b/Engine/Plugins/VirtualProduction/VirtualCamera/Source/VirtualCameraEditor/Private/Cinematic/MovieSceneSpawnableVCamBinding.cpp
--- a/Engine/Plugins/VirtualProduction/VirtualCamera/Source/VirtualCameraEditor/Private/Cinematic/MovieSceneSpawnableVCamBinding.cpp
+++ b/Engine/Plugins/VirtualProduction/VirtualCamera/Source/VirtualCameraEditor/Private/Cinematic/MovieSceneSpawnableVCamBinding.cpp
@@ -40,6 +40,7 @@ UMovieSceneCustomBinding* UMovieSceneSpawnableVCamBinding::CreateNewCustomBindin
 	//  - In inner, uncropped viewport which is supposed to be what the CineCameraActor is supposed to see after recording is done.
 	//  - The outer, cropped viewport, which is darker in the VCam UI. That section would be visible if we didn't reset asymmetric overscan here.
 	CineCameraActor->GetCineCameraComponent()->AsymmetricOverscan = FVector4f::Zero();
+	CineCameraActor->GetCineCameraComponent()->bConstrainAspectRatio = true;
 	
 	CineCameraActor->Tags.Remove(TEXT("SequencerActor"));
 	CineCameraActor->Tags.Remove(TEXT("SequencerPreviewActor"));
