Virtual Camera: When creating template actor for take recording, skip delta serialization so all properties are correctly duplicated into the template actor. This fixes post process overrides sometimes not being correctly recorded.

diff --git a/Engine/Plugins/VirtualProduction/VirtualCamera/Source/VirtualCameraEditor/Private/Cinematic/MovieSceneSpawnableVCamBinding.cpp b/Engine/Plugins/VirtualProduction/VirtualCamera/Source/VirtualCameraEditor/Private/Cinematic/MovieSceneSpawnableVCamBinding.cpp
--- a/Engine/Plugins/VirtualProduction/VirtualCamera/Source/VirtualCameraEditor/Private/Cinematic/MovieSceneSpawnableVCamBinding.cpp
+++ b/Engine/Plugins/VirtualProduction/VirtualCamera/Source/VirtualCameraEditor/Private/Cinematic/MovieSceneSpawnableVCamBinding.cpp
@@ -33,6 +33,7 @@ UMovieSceneCustomBinding* UMovieSceneSpawnableVCamBinding::CreateNewCustomBindin
 	CopyParams.bNotifyObjectReplacement = false;
 	CopyParams.bPreserveRootComponent = false;
 	CopyParams.bPerformDuplication = true;
+	CopyParams.bDoDelta = false; // Required so all properties are copied over correctly. 
 	UEngine::CopyPropertiesForUnrelatedObjects(SourceObject, CineCameraActor, CopyParams);
 
 	// Specific for VCam: The asymmetric overscan needs to be zero in the resulting level sequence.
