Virtual Camera: Move UMovieSceneSpawnableVCamBinding to runtime module so VCam can be recorded in -game

diff --git a/Engine/Plugins/VirtualProduction/VirtualCamera/Source/VirtualCamera/Private/Cinematic/MovieSceneSpawnableVCamBinding.cpp b/Engine/Plugins/VirtualProduction/VirtualCamera/Source/VirtualCamera/Private/Cinematic/MovieSceneSpawnableVCamBinding.cpp
new file mode 100644
--- /dev/null
+++ b/Engine/Plugins/VirtualProduction/VirtualCamera/Source/VirtualCamera/Private/Cinematic/MovieSceneSpawnableVCamBinding.cpp
@@ -0,0 +1,60 @@
+﻿// Copyright Epic Games, Inc. All Rights Reserved.
+
+#include "MovieSceneSpawnableVCamBinding.h"
+
+#include "CineCameraActor.h"
+#include "Cinematic/VCamHierarchyInfo.h"
+#include "Engine/Engine.h"
+#include "MovieScene.h"
+
+bool UMovieSceneSpawnableVCamBinding::SupportsBindingCreationFromObject(const UObject* SourceObject) const
+{
+	return SourceObject && UE::VirtualCamera::FVCamHierarchyInfo(SourceObject).ShouldRecordAsCineCamera();
+}
+
+UMovieSceneCustomBinding* UMovieSceneSpawnableVCamBinding::CreateNewCustomBinding(UObject* SourceObject, UMovieScene& OwnerMovieScene)
+{
+	const UE::VirtualCamera::FVCamHierarchyInfo VCamInfo(SourceObject);
+	if (!SourceObject || !VCamInfo.ShouldRecordAsCineCamera()) // We don't expect this to happen as SupportsBindingCreationFromObject rejects this case.
+	{
+		return nullptr;
+	}
+	
+	const FName TemplateName = MakeUniqueObjectName(&OwnerMovieScene, UObject::StaticClass(), SourceObject ? SourceObject->GetFName() : TEXT("EmptyBinding"));
+	const FName InstancedBindingName = MakeUniqueObjectName(&OwnerMovieScene, UObject::StaticClass(), *FString(TemplateName.ToString() + TEXT("_CustomBinding")));
+
+	ACineCameraActor* CineCameraActor = NewObject<ACineCameraActor>(&OwnerMovieScene, TemplateName, RF_Transactional);
+	if (!CineCameraActor)
+	{
+		return nullptr;
+	}
+	
+	UEngine::FCopyPropertiesForUnrelatedObjectsParams CopyParams;
+	CopyParams.bNotifyObjectReplacement = false;
+	CopyParams.bPreserveRootComponent = false;
+	CopyParams.bPerformDuplication = true;
+	CopyParams.bDoDelta = false; // Required so all properties are copied over correctly. 
+	UEngine::CopyPropertiesForUnrelatedObjects(SourceObject, CineCameraActor, CopyParams);
+
+	// Specific for VCam: The asymmetric overscan needs to be zero in the resulting level sequence.
+	// The SensorCorrectionModifier sets the asymmetric overscan to be non-zero to achieve two "viewports" in the VCam UI.
+	//  - In inner, uncropped viewport which is supposed to be what the CineCameraActor is supposed to see after recording is done.
+	//  - The outer, cropped viewport, which is darker in the VCam UI. That section would be visible if we didn't reset asymmetric overscan here.
+	CineCameraActor->GetCineCameraComponent()->AsymmetricOverscan = FVector4f::Zero();
+	CineCameraActor->GetCineCameraComponent()->bConstrainAspectRatio = true;
+	
+	CineCameraActor->Tags.Remove(TEXT("SequencerActor"));
+	CineCameraActor->Tags.Remove(TEXT("SequencerPreviewActor"));
+	CineCameraActor->DetachFromActor(FDetachmentTransformRules(EDetachmentRule::KeepRelative, false));
+#if WITH_EDITOR
+	CineCameraActor->bIsEditorPreviewActor = false;
+#endif
+
+	// Record as if this was CineCameraActor.
+	// Do not create UMovieSceneSpawnableVCamBinding instance because that class won't be available if the user disables the Virtual Camera plugin.
+	UMovieSceneSpawnableActorBinding* Binding = NewObject<UMovieSceneSpawnableActorBinding>(
+		&OwnerMovieScene, UMovieSceneSpawnableActorBinding::StaticClass(), InstancedBindingName, RF_Transactional
+		);
+	Binding->SetObjectTemplate(CineCameraActor);
+	return Binding;
+}
\ No newline at end of file
diff --git a/Engine/Plugins/VirtualProduction/VirtualCamera/Source/VirtualCamera/Private/Cinematic/MovieSceneSpawnableVCamBinding.h b/Engine/Plugins/VirtualProduction/VirtualCamera/Source/VirtualCamera/Private/Cinematic/MovieSceneSpawnableVCamBinding.h
new file mode 100644
--- /dev/null
+++ b/Engine/Plugins/VirtualProduction/VirtualCamera/Source/VirtualCamera/Private/Cinematic/MovieSceneSpawnableVCamBinding.h
@@ -0,0 +1,36 @@
+﻿// Copyright Epic Games, Inc. All Rights Reserved.
+
+#pragma once
+
+#include "Bindings/MovieSceneSpawnableActorBinding.h"
+#include "MovieSceneSpawnableVCamBinding.generated.h"
+
+namespace UE::TakeRecorderSources{ struct FCanRecordArgs; }
+
+/**
+ * Allows the AVCamBaseActor to be recorded as ACineCameraActor instead.
+ *
+ * This allows users to record a VCam actor using Take Recorder and disable the VCam plugin afterwards while still being able to use the recorded
+ * sequences without missing classes.
+ *
+ * This class is never actually instantiated. It simply creates a UMovieSceneSpawnableActorBinding that contains a ACineCameraActor instead of a VCam.
+ * Not instantiating this class is important: otherwise the sequence referencing the binding would lose the reference when VCam plugin is disabled as
+ * the class would be missing.
+ */
+UCLASS()
+class UMovieSceneSpawnableVCamBinding : public UMovieSceneSpawnableActorBinding
+{
+	GENERATED_BODY()
+public:
+
+	//~ Begin UMovieSceneSpawnableActorBinding Interface
+	virtual bool SupportsBindingCreationFromObject(const UObject* SourceObject) const override;
+	virtual UMovieSceneCustomBinding* CreateNewCustomBinding(UObject* SourceObject, UMovieScene& OwnerMovieScene) override;
+	//~ End UMovieSceneSpawnableActorBinding Interface
+
+protected:
+	
+	//~ Begin UMovieSceneSpawnableActorBinding Interface
+	virtual int32 GetCustomBindingPriority() const override { return BaseCustomPriority; }
+	//~ End UMovieSceneSpawnableActorBinding Interface
+};
\ No newline at end of file
diff --git a/Engine/Plugins/VirtualProduction/VirtualCamera/Source/VirtualCamera/Public/Cinematic/VCamHierarchyInfo.h b/Engine/Plugins/VirtualProduction/VirtualCamera/Source/VirtualCamera/Public/Cinematic/VCamHierarchyInfo.h
new file mode 100644
--- /dev/null
+++ b/Engine/Plugins/VirtualProduction/VirtualCamera/Source/VirtualCamera/Public/Cinematic/VCamHierarchyInfo.h
@@ -0,0 +1,30 @@
+﻿// Copyright Epic Games, Inc. All Rights Reserved.
+
+#pragma once
+
+#include "CineCameraComponent.h"
+#include "GameFramework/Actor.h"
+#include "VCamComponent.h"
+#include "Misc/NotNull.h"
+
+namespace UE::VirtualCamera
+{
+/** Gets components of a VCamActor for the purpose of deciding whether it should be recorded as cine camera in Sequencer. */
+struct FVCamHierarchyInfo
+{
+	const AActor* Actor;
+	UVCamComponent* VCamComponent;
+
+	USceneComponent* RootComponent;
+	UCineCameraComponent* Camera;
+
+	explicit FVCamHierarchyInfo(TNotNull<const UObject*> SourceObject)
+		: Actor(SourceObject->IsA<AActor>() ? Cast<AActor>(SourceObject) : SourceObject->GetTypedOuter<AActor>())
+		, VCamComponent(Actor ? Actor->FindComponentByClass<UVCamComponent>() : nullptr)
+		, RootComponent(Actor ? Actor->GetRootComponent() : nullptr)
+		, Camera(VCamComponent ? VCamComponent->GetTargetCamera() : nullptr)
+	{}
+
+	bool ShouldRecordAsCineCamera() const { return VCamComponent != nullptr && Camera != nullptr && VCamComponent->GetRecordAsCineCamera(); }
+};
+}
\ No newline at end of file
diff --git a/Engine/Plugins/VirtualProduction/VirtualCamera/Source/VirtualCamera/VirtualCamera.Build.cs b/Engine/Plugins/VirtualProduction/VirtualCamera/Source/VirtualCamera/VirtualCamera.Build.cs
--- a/Engine/Plugins/VirtualProduction/VirtualCamera/Source/VirtualCamera/VirtualCamera.Build.cs
+++ b/Engine/Plugins/VirtualProduction/VirtualCamera/Source/VirtualCamera/VirtualCamera.Build.cs
@@ -34,7 +34,8 @@ public class VirtualCamera : ModuleRules
 			{
 				"MediaIOCore",
 				"Slate",
-				"AdvancedWidgets"
+				"AdvancedWidgets",
+				"MovieSceneTracks",
 			}
 		);
 		
diff --git a/Engine/Plugins/VirtualProduction/VirtualCamera/Source/VirtualCameraEditor/Private/Cinematic/MovieSceneSpawnableVCamBinding.cpp b/Engine/Plugins/VirtualProduction/VirtualCamera/Source/VirtualCameraEditor/Private/Cinematic/MovieSceneSpawnableVCamBinding.cpp
deleted file mode 100644
diff --git a/Engine/Plugins/VirtualProduction/VirtualCamera/Source/VirtualCameraEditor/Private/Cinematic/MovieSceneSpawnableVCamBinding.h b/Engine/Plugins/VirtualProduction/VirtualCamera/Source/VirtualCameraEditor/Private/Cinematic/MovieSceneSpawnableVCamBinding.h
deleted file mode 100644
diff --git a/Engine/Plugins/VirtualProduction/VirtualCamera/Source/VirtualCameraEditor/Private/Cinematic/VCamHierarchyInfo.h b/Engine/Plugins/VirtualProduction/VirtualCamera/Source/VirtualCameraEditor/Private/Cinematic/VCamHierarchyInfo.h
deleted file mode 100644
diff --git a/Engine/Plugins/VirtualProduction/VirtualCamera/Source/VirtualCameraEditor/Private/Cinematic/VCamSupportForCinematicTooling.cpp b/Engine/Plugins/VirtualProduction/VirtualCamera/Source/VirtualCameraEditor/Private/Cinematic/VCamSupportForCinematicTooling.cpp
--- a/Engine/Plugins/VirtualProduction/VirtualCamera/Source/VirtualCameraEditor/Private/Cinematic/VCamSupportForCinematicTooling.cpp
+++ b/Engine/Plugins/VirtualProduction/VirtualCamera/Source/VirtualCameraEditor/Private/Cinematic/VCamSupportForCinematicTooling.cpp
@@ -3,7 +3,7 @@
 #include "VCamSupportForCinematicTooling.h"
 
 #include "ITakeRecorderSourcesModule.h"
-#include "VCamHierarchyInfo.h"
+#include "Cinematic/VCamHierarchyInfo.h"
 
 namespace UE::VirtualCamera
 {
diff --git a/Engine/Plugins/VirtualProduction/VirtualCamera/Source/VirtualCameraEditor/VirtualCameraEditor.Build.cs b/Engine/Plugins/VirtualProduction/VirtualCamera/Source/VirtualCameraEditor/VirtualCameraEditor.Build.cs
--- a/Engine/Plugins/VirtualProduction/VirtualCamera/Source/VirtualCameraEditor/VirtualCameraEditor.Build.cs
+++ b/Engine/Plugins/VirtualProduction/VirtualCamera/Source/VirtualCameraEditor/VirtualCameraEditor.Build.cs
@@ -16,8 +16,6 @@ public class VirtualCameraEditor : ModuleRules
 				"CinematicCamera",
 				"EditorWidgets",
 				"Engine",
-				"MovieScene",
-				"MovieSceneTracks",
 				"PlacementMode", 
 				"Settings",
 				"UnrealEd",
